Move GDK plugins to the public engine

#rnx

#Synced-CL 50857758

diff --git a/Engine/Plugins/Experimental/Microsoft/GDKDLCSupport/GDKDLCSupport.uplugin b/Engine/Plugins/Experimental/Microsoft/GDKDLCSupport/GDKDLCSupport.uplugin
new file mode 100644
--- /dev/null
+++ b/Engine/Plugins/Experimental/Microsoft/GDKDLCSupport/GDKDLCSupport.uplugin
@@ -0,0 +1,30 @@
+{
+	"FileVersion": 3,
+	"Version": 1,
+	"VersionName": "0.1",
+	"FriendlyName": "DLC Support for GDK",
+	"Description": "Experimental support for DLC in GDK.",
+	"CreatedBy": "Epic Games, Inc.",
+	"Category": "Platform",
+	"CreatedByURL": "https://epicgames.com",
+	"DocsURL": "",
+	"MarketplaceURL": "",
+	"SupportURL": "",
+	"EnabledByDefault": false,
+	"RequiresBuildPlatform": false,
+	"CanContainContent": false,
+	"IsBetaVersion": true,
+	"Installed": false,
+	"HasExplicitPlatforms": true,	
+	"SupportedTargetPlatforms": [ "Win64" ],
+	"Modules": [
+		{
+			"Name": "GDKDLCSupport",
+			"Type": "RuntimeNoCommandlet",
+			"LoadingPhase": "PostConfigInit",
+			"HasExplicitPlatforms": true,
+			"PlatformAllowList": [ "Win64" ],
+			"PlatformArchitectureDenyList": [ "Win64:arm64" ]
+		}
+	]
+}
\ No newline at end of file
diff --git a/Engine/Plugins/Experimental/Microsoft/GDKDLCSupport/Source/GDKDLCSupport/GDKDLCSupport.Build.cs b/Engine/Plugins/Experimental/Microsoft/GDKDLCSupport/Source/GDKDLCSupport/GDKDLCSupport.Build.cs
new file mode 100644
--- /dev/null
+++ b/Engine/Plugins/Experimental/Microsoft/GDKDLCSupport/Source/GDKDLCSupport/GDKDLCSupport.Build.cs
@@ -0,0 +1,23 @@
+// Copyright Epic Games, Inc. All Rights Reserved.
+
+using UnrealBuildTool;
+using EpicGames.Core;
+
+public class GDKDLCSupport : ModuleRules
+{
+	public GDKDLCSupport(ReadOnlyTargetRules Target) : base(Target)
+	{
+		PrivateDependencyModuleNames.AddAll( 
+			"GRDK",
+			"GDKRuntime",
+            "GDKPackageManifest",	
+			"Core",
+			"CoreUObject",
+			"Engine",
+			"Projects",
+			"AssetRegistry"
+		);
+
+		PublicIncludePathModuleNames.Add("GDKRuntime");
+	}
+}
diff --git a/Engine/Plugins/Experimental/Microsoft/GDKDLCSupport/Source/GDKDLCSupport/Private/GDKDLCHandler.cpp b/Engine/Plugins/Experimental/Microsoft/GDKDLCSupport/Source/GDKDLCSupport/Private/GDKDLCHandler.cpp
new file mode 100644
--- /dev/null
+++ b/Engine/Plugins/Experimental/Microsoft/GDKDLCSupport/Source/GDKDLCSupport/Private/GDKDLCHandler.cpp
@@ -0,0 +1,1077 @@
+// Copyright Epic Games, Inc. All Rights Reserved.
+
+#include "GDKDLCHandler.h"
+#if WITH_GRDK
+#include "GDKTaskQueueHelpers.h"
+#include "GDKThreadCheck.h"
+#include "GDKRuntimeModule.h"
+#include "IGDKPackageManifestModule.h"
+#include "HAL/PlatformFile.h"
+#include "HAL/PlatformFileManager.h"
+#include "HAL/IConsoleManager.h"
+#include "Misc/CoreDelegates.h"
+#include "Misc/ConfigCacheIni.h"
+#include "Misc/Paths.h"
+#include "Misc/ScopeLock.h"
+#include "Misc/CString.h"
+#include "Async/Async.h"
+#include "Interfaces/IPluginManager.h"
+#include "AssetRegistry/IAssetRegistry.h"
+#include "AssetRegistry/AssetRegistryState.h"
+
+DEFINE_LOG_CATEGORY(LogGDKDLC);
+
+static bool GGDKDLCMountAgeRestrictedDLC = false;
+static FAutoConsoleVariableRef CVarDLCMountAgeRestrictedDLC(
+	TEXT( "GDK.DLC.MountAgeRestricted" ),
+	GGDKDLCMountAgeRestrictedDLC,
+	TEXT( "Whether to attempt to mount age-restricted DLC. This will trigger a parental controls consent prompt." ));
+
+static bool GGDKDLCAutoMountOnStartupDLC = true;
+static FAutoConsoleVariableRef CVarDLCAutoMountOnStartupDLC(
+	TEXT( "GDK.DLC.MountOnStartup" ),
+	GGDKDLCAutoMountOnStartupDLC,
+	TEXT( "Whether to automatically mount any downloaded DLC on system startup." ),
+	ECVF_ReadOnly);
+
+static bool GGDKDLCAutoMountRuntimeDLC = true;
+static FAutoConsoleVariableRef CVarDLCAutoMountRuntimeDLC(
+	TEXT( "GDK.DLC.MountAtRuntime" ),
+	GGDKDLCAutoMountRuntimeDLC,
+	TEXT( "Whether to automatically mount any new DLC that is downloaded once the application is up and running." ),
+	ECVF_ReadOnly);
+
+static bool GGDKDLCAutoAcquireLicenses = true;
+static FAutoConsoleVariableRef CVarDLCAutoAcquireLicenses(
+	TEXT( "GDK.DLC.AutoAcquireLicense" ),
+	GGDKDLCAutoAcquireLicenses,
+	TEXT( "Whether to attempt to acquire the license for any DLC before attempting to mount. If this is false, the title must manage licensing directly." ));
+
+#if !UE_BUILD_SHIPPING
+static bool GGDKDLCIgnoreDevelopmentLicenseIssues = true;
+static FAutoConsoleVariableRef CVarDLCIgnoreDevelopmentLicenseIssues(
+	TEXT( "GDK.DLC.IgnoreDevelopmentLicenseIssues" ),
+	GGDKDLCIgnoreDevelopmentLicenseIssues,
+	TEXT( "Whether to ignore certain failure states when license aquisition fails in a non-shipping build - typically during dev sideloading scenarios." ));
+#endif
+
+FGDKDLCHandler::FGDKDLCHandler()
+	: CurrentStoreContext(nullptr)
+{
+	GDK_SCOPE_NOT_TIME_SENSITIVE(); // (startup only) XPackageEnumeratePackages is not safe to call on a time-sensitive thread
+
+	// make sure the package manifest module is loaded
+	IGDKPackageManifestModule::Get();
+
+	// use the default user for acquiring DLC licenses from the store
+	SetStoreUser( FPlatformUserId::CreateFromInternalId(0) );
+
+	// create package data for any DLC
+	auto EnumeratePackageCallback = [](void* Context, const XPackageDetails* Details)
+	{
+		FGDKDLCHandler* This = (FGDKDLCHandler*)Context;
+		This->AsyncInitDLC(*Details, GGDKDLCAutoMountOnStartupDLC);
+		return true;
+	};
+	XPackageEnumeratePackages( XPackageKind::Content, XPackageEnumerationScope::ThisOnly, this, EnumeratePackageCallback );
+
+	// register a callback when new packages are downloaded
+	auto PackageInstalledCallback = [](void* Context, const XPackageDetails* Details)
+	{
+		FGDKDLCHandler* This = (FGDKDLCHandler*)Context;
+		This->AsyncInitDLC(*Details, GGDKDLCAutoMountRuntimeDLC);
+	};
+	XPackageRegisterPackageInstalled(FGDKAsyncTaskQueue::GetGenericQueue(), this, PackageInstalledCallback, &PackageInstalledToken);
+}
+
+
+FGDKDLCHandler::~FGDKDLCHandler()
+{
+	// try to stop any current processing task
+	ProcessingState = EProcessingState::Destroying;
+	CurrentProcessingAsyncTask.TryCancel(true);
+	PendingPackages.Reset();
+	LostLicensePackages.Reset();
+	CurrentProcessingPackage.Reset();
+
+	// release all packages (only the package captured by AsyncGDKTask blocks will be kept alive)
+	AllPackages.Reset();
+
+
+	XPackageUnregisterPackageInstalled(PackageInstalledToken, true);
+
+	if (CurrentStoreContext != nullptr)
+	{
+		XStoreCloseContextHandle(CurrentStoreContext);
+		CurrentStoreContext = nullptr;
+	}
+
+	RemoveFromTicker();
+}
+
+
+bool FGDKDLCHandler::Tick( float DeltaSeconds )
+{
+	switch (ProcessingState)
+	{
+		case EProcessingState::Idle:
+		{
+			FScopeLock Lock(&PackageDataLock);
+
+			check(!CurrentProcessingPackage.IsValid());
+			if (PendingPackages.Num() > 0)
+			{
+				// get the next package to mount
+				CurrentProcessingPackage = PendingPackages.Pop();
+				InitializePackage_GameThread(CurrentProcessingPackage);
+			}
+			else if (LostLicensePackages.Num() > 0)
+			{
+				// get the next lost license package
+				CurrentProcessingPackage = LostLicensePackages.Pop();
+				LicenseLost_GameThread(CurrentProcessingPackage);
+				CurrentProcessingPackage.Reset();
+			}
+			else
+			{
+				// there's no more work to do
+				RemoveFromTicker();
+				return false;
+			}
+		}
+		break;
+
+		case EProcessingState::QueryingLicense:
+		case EProcessingState::AcquiringLicense:
+		case EProcessingState::MountingPackage:
+		case EProcessingState::CachingData:
+		{
+			// nothing to do, but sanity check that the package is still valid
+			check(CurrentProcessingPackage.IsValid());
+		}
+		break;
+
+		case EProcessingState::MountingPlugin:
+		{
+			check(CurrentProcessingPackage.IsValid());
+			MountPlugin_GameThread(CurrentProcessingPackage);
+			check(ProcessingState == EProcessingState::Completed);
+		}
+		//... fall through to completion ...
+
+		case EProcessingState::Completed:
+		{
+			check(CurrentProcessingPackage.IsValid());
+			Complete_GameThread(CurrentProcessingPackage);
+			check(!CurrentProcessingPackage.IsValid());
+			check(ProcessingState == EProcessingState::Idle);
+		}
+		break;
+	}
+
+	// keep ticking
+	return true;
+}
+
+
+void FGDKDLCHandler::AsyncInitDLC(const XPackageDetails& Details, bool bMount)
+{
+	UE_CLOG(Details.installing, LogGDKDLC, Error, TEXT("package %s is still installing - intelligent delivery is not supported for DLC at the moment"), UTF8_TO_TCHAR(Details.packageIdentifier) );
+
+	// add package and prepare for enumeration callbacks
+	TSharedPtr<FPackageData> Package;
+	{
+		FScopeLock Lock(&PackageDataLock);
+		Package = AllPackages.Add_GetRef( MakeShared<FPackageData>(this, Details));
+		Package->Context.Package = Package; // initialize the callback context. necessary to support those GDK callback functions that only take a raw context pointer
+	}
+
+	// mount the DLC package automatically, if requested (this is typically done for DLC that arrives after engine startup)
+	if (bMount)
+	{
+		AddToPendingPackages( Package );
+	}
+}
+
+
+void FGDKDLCHandler::AddToPendingPackages( TSharedPtr<FPackageData> Package )
+{
+	FScopeLock Lock(&PackageDataLock);
+	PendingPackages.AddUnique(Package);
+	AddToTicker();
+}
+
+
+void FGDKDLCHandler::InitializePackage_GameThread( TSharedPtr<FPackageData> Package )
+{
+	// if the package is already mounted and ready, just fire the callbacks again
+	bool bHasLicense = (Package->LicenseState == ELicenseState::Licensed || Package->LicenseState == ELicenseState::Unlicensed_DevIgnore);
+	if (Package->MountState == EMountState::Mounted && bHasLicense)
+	{
+		ProcessingState = EProcessingState::Completed;
+		return;
+	}
+
+	// check for valid state
+	if (bHasLicense && Package->MountState != EMountState::NotMounted && Package->MountState != EMountState::Unmounted)
+	{
+		checkNoEntry();
+		return;
+	}
+
+	// early out if we are ignoring age-restricted DLC
+	if (Package->bIsAgeRestricted && !GGDKDLCMountAgeRestrictedDLC)
+	{
+		Package->MountState = EMountState::Ignored_AgeRestricted;
+		ProcessingState = EProcessingState::Completed;
+		UE_LOG(LogGDKDLC, Log, TEXT("Ignoring DLC %s because it is age restricted for one or more of the currently logged in users"), *Package->PackageName );
+		return;
+	}
+
+	// see if we should start acquiring a license
+	bool bWantLicense = false;
+	if (GGDKDLCAutoAcquireLicenses)
+	{
+		// make sure there is a store context
+		if (!EnsureValidStoreContext())
+		{
+			ProcessingState = EProcessingState::Completed;
+			UE_LOG(LogGDKDLC, Error, TEXT("cannot mount DLC"));
+			return;
+		}
+
+		// see if we should acquire a license for this DLC
+		bWantLicense = (CurrentStoreContext != nullptr && Package->DLCLicense == nullptr);
+	}
+
+	if (bWantLicense)
+	{
+		Package->StoreUserId = CurrentStoreUserId;
+		ProcessingState = EProcessingState::QueryingLicense;
+		AsyncQueryLicense_GameThread(Package);
+	}
+	else
+	{
+		ProcessingState = EProcessingState::MountingPackage;
+		AsyncMountPackage_GameThread(Package);
+	}
+}
+
+
+void FGDKDLCHandler::AsyncQueryLicense_GameThread( TSharedPtr<FPackageData> Package )
+{
+	check(Package->DLCLicense == nullptr);
+	check(ProcessingState == EProcessingState::QueryingLicense);
+
+	HRESULT hResult = AsyncGDKTask(
+		CurrentProcessingAsyncTask,
+		[Package, this](XAsyncBlock* Block)
+		{
+			return XStoreCanAcquireLicenseForPackageAsync(CurrentStoreContext, Package->PackageIdentifier, Block);
+		},
+		[Package, this](XAsyncBlock* Block)
+		{
+			// early out if the DLC system is being destroyed
+			bool bCancel = (ProcessingState == EProcessingState::Destroying);
+			if (bCancel)
+			{
+				return;
+			}
+
+			XStoreCanAcquireLicenseResult Result;
+			HRESULT hResult = XStoreCanAcquireLicenseForPackageResult(Block, &Result);
+			UE_CLOG(FAILED(hResult), LogGDKDLC, Error, TEXT("failed to query license for package %s : 0x%X"), *Package->PackageName, hResult);
+
+			if (SUCCEEDED(hResult) && Result.status != XStoreCanLicenseStatus::NotLicensableToUser)
+			{
+				ProcessingState = EProcessingState::AcquiringLicense;
+				AsyncAcquireLicense_GameThread(Package);
+			}
+			else
+			{
+				Package->LicenseState = ELicenseState::Unlicensed;
+				ProcessingState = EProcessingState::Completed;
+				UE_CLOG(SUCCEEDED(hResult), LogGDKDLC, Error, TEXT("Package %s is not licensable for the current store user (result = %d)"), *Package->PackageName, Result.status);
+			}
+		}
+	);
+
+	if (FAILED(hResult))
+	{
+		ProcessingState = EProcessingState::Completed;
+		UE_LOG(LogGDKDLC, Error, TEXT("Failed to start querying licensability for package %s : 0x%X"), *Package->PackageName, hResult);
+	}
+}
+
+
+void FGDKDLCHandler::AsyncAcquireLicense_GameThread( TSharedPtr<FPackageData> Package )
+{
+	check(Package->DLCLicense == nullptr);
+	check(ProcessingState == EProcessingState::AcquiringLicense);
+
+	HRESULT hResult = AsyncGDKTask(
+		CurrentProcessingAsyncTask,
+		[Package, this](XAsyncBlock* Block)
+		{
+			return XStoreAcquireLicenseForPackageAsync(CurrentStoreContext, Package->PackageIdentifier, Block);
+		},
+		[Package, this](XAsyncBlock* Block)
+		{
+			HRESULT hResult = XStoreAcquireLicenseForPackageResult(Block, &Package->DLCLicense);
+
+			// early out if the DLC system is being destroyed (must be done after caching the DLC license so we don't leak it)
+			bool bCancel = (ProcessingState == EProcessingState::Destroying);
+			if (bCancel)
+			{
+				return;
+			}
+
+			if (SUCCEEDED(hResult) && XStoreIsLicenseValid(Package->DLCLicense) )
+			{
+				Package->LicenseState = ELicenseState::Licensed;
+
+				// register for a callback when the license is lost
+				hResult = XStoreRegisterPackageLicenseLost(Package->DLCLicense, FGDKAsyncTaskQueue::GetGenericQueue(), &Package->Context, &FGDKDLCHandler::OnPackageLicenseLost, &Package->PackageLicenseLostToken );
+				UE_CLOG(FAILED(hResult), LogGDKDLC, Warning, TEXT("Could not register for GDK license lost callback : 0x%X"), hResult);
+
+				// start mounting the DLC package
+				ProcessingState = EProcessingState::MountingPackage;
+				AsyncMountPackage_GameThread(Package);
+				return;
+			}
+
+			// if we received a license that was not valid, close it now as there's nothing we can do with it
+			if (Package->DLCLicense != nullptr)
+			{
+				XStoreCloseLicenseHandle(Package->DLCLicense);
+				Package->DLCLicense = nullptr;
+			}
+
+			// handle failures
+#if !UE_BUILD_SHIPPING
+			if (GGDKDLCIgnoreDevelopmentLicenseIssues && (
+				hResult == 0x87e10bc6 || // LM_E_CONTENT_NOT_IN_CATALOG
+				hResult == 0x803f9006 || // LM_E_ENTITLED_USER_SIGNED_OUT
+				hResult == 0x87dd0013))  // AM_E_USER_NOT_FOUND
+			{
+				// this is most likely from using the GDK development tools to sideload the DLC package instead of aquiring it from the Store. Allow the mount to happen regardless
+				UE_LOG(LogGDKDLC, Warning, TEXT("Ignoring DLC license aquisition failure (0x%X) in non-shipping build. This will fail in Shipping"), hResult );
+				Package->LicenseState = ELicenseState::Unlicensed_DevIgnore;
+
+				ProcessingState = EProcessingState::MountingPackage;
+				AsyncMountPackage_GameThread(Package);
+			}
+			else
+#endif
+			{
+				Package->LicenseState = ELicenseState::Unlicensed;
+				ProcessingState = EProcessingState::Completed;
+				UE_LOG(LogGDKDLC, Error, TEXT("Failed to get license for package %s : 0x%X"), *Package->PackageName, hResult);
+			}
+		}
+	);
+
+	if (FAILED(hResult))
+	{
+		Package->LicenseState = ELicenseState::Unlicensed;
+		ProcessingState = EProcessingState::Completed;
+		UE_LOG(LogGDKDLC, Error, TEXT("Failed to start getting license for package %s : 0x%X"), *Package->PackageName, hResult);
+	}
+}
+
+
+void FGDKDLCHandler::AsyncMountPackage_GameThread( TSharedPtr<FPackageData> Package )
+{
+	check(ProcessingState == EProcessingState::MountingPackage);
+
+	// if the package is already mounted, it means we've only reaquired the license
+	if (Package->DLCHandle != nullptr && ensure(Package->MountState == EMountState::Mounted) )
+	{
+		ProcessingState = EProcessingState::Completed;
+		return;
+	}
+
+
+	HRESULT hResult = AsyncGDKTask(
+		CurrentProcessingAsyncTask,
+		[Package](XAsyncBlock* Block)
+		{
+			return XPackageMountWithUiAsync(Package->PackageIdentifier, Block);
+		},
+		[Package, this](XAsyncBlock* Block)
+		{
+			HRESULT hResult = XPackageMountWithUiResult(Block, &Package->DLCHandle);
+
+			// early out if the DLC system is being destroyed (must be done after caching the DLC handle so we don't leak it)
+			bool bCancel = (ProcessingState == EProcessingState::Destroying);
+			if (bCancel)
+			{
+				return;
+			}
+
+			if (SUCCEEDED(hResult))
+			{
+				// continue to initialize the DLC package on a background thread
+				ProcessingState = EProcessingState::CachingData;
+				AsyncTask( ENamedThreads::AnyBackgroundThreadNormalTask, [Package, this]()
+				{
+					CachePackageData_BackgroundThread(Package);
+				});
+				return;
+			}
+
+			Package->MountState = EMountState::NotMounted;
+			ProcessingState = EProcessingState::Completed;
+			UE_LOG(LogGDKDLC, Error, TEXT("Failed to mount DLC '%s' : 0x%X"), *Package->PackageName, hResult);
+		}
+	);
+
+	if (FAILED(hResult))
+	{
+		Package->MountState = EMountState::NotMounted;
+		ProcessingState = EProcessingState::Completed;
+		UE_LOG(LogGDKDLC, Error, TEXT("Failed to start mounting mount DLC '%s' : 0x%X"), *Package->PackageName, hResult);
+	}
+}
+
+
+void FGDKDLCHandler::CachePackageData_BackgroundThread( TSharedPtr<FPackageData> Package )
+{
+	check(ProcessingState == EProcessingState::CachingData);
+
+	HRESULT hResult;
+
+	// get the mount path
+	char MountPath[MAX_PATH+1];
+	hResult = XPackageGetMountPath(Package->DLCHandle, UE_ARRAY_COUNT(MountPath)-1, MountPath);
+	if (FAILED(hResult))
+	{
+		Package->MountState = EMountState::NotMounted;
+		ProcessingState = EProcessingState::Completed;
+		UE_LOG(LogGDKDLC, Error, TEXT("Failed to get mount path for DLC '%s' : 0x%X"), *Package->PackageName, hResult);
+
+		XPackageCloseMountHandle(Package->DLCHandle);
+		Package->DLCHandle = nullptr;
+		return;
+	}
+	Package->MountPath = UTF8_TO_TCHAR(MountPath);
+
+
+	// load package manifest
+	IGDKPackageManifestModule::Get().LoadManifestFromDLC(Package->DLCHandle, Package->PackageIdentifier, Package->MountPath);
+
+	// finalize on the GameThread
+	ProcessingState = EProcessingState::MountingPlugin;
+}
+
+
+void FGDKDLCHandler::MountPlugin_GameThread( TSharedPtr<FPackageData> Package )
+{
+	check(ProcessingState == EProcessingState::MountingPlugin);
+
+	// mount all available pak files. NOTE: intelligent delivery is not supported for DLC at the moment - only what is currently installed will be mounted
+	TArray<FString> PakFileNames;
+	IPlatformFile::GetPlatformPhysical().FindFilesRecursively(PakFileNames, *Package->MountPath, TEXT(".pak") );
+	for (const FString& PakFileName : PakFileNames)
+	{
+		if (IGDKPackageManifestModule::Get().IsPakFileInstalled(Package->DLCHandle, PakFileName))
+		{
+			IPakFile* PakFile = FCoreDelegates::MountPak.IsBound() ? FCoreDelegates::MountPak.Execute(PakFileName, INDEX_NONE) : nullptr;
+			if (PakFile)
+			{
+				UE_LOG(LogGDKDLC, Log, TEXT("Successfully mounted %s"), *PakFileName);
+			}
+			else
+			{
+				UE_LOG(LogGDKDLC, Log, TEXT("Failed to mount %s"), *PakFileName);
+			}
+		}
+	}
+
+	// find & load the .uplugin (should only be one per DLC, as specified by "RunUAT ... -DLCName=<plugin>")
+	TArray<FString> FoundPlugins;
+	IPlatformFile::GetPlatformPhysical().FindFilesRecursively(FoundPlugins, *Package->MountPath, TEXT(".uplugin") );
+	checkf(FoundPlugins.Num() != 0, TEXT("found %d plugins in %s - DLC should have 1!"), FoundPlugins.Num(), *Package->MountPath);
+	const FString PluginFileName = FoundPlugins.Last();
+
+	IPluginManager::Get().AddToPluginsList(PluginFileName);
+	FString DLCName = FPaths::GetBaseFilename(PluginFileName);
+
+	Package->DLCName = FName(DLCName);
+	Package->DLCPlugin = IPluginManager::Get().FindPlugin(DLCName);
+	checkf(Package->DLCPlugin.IsValid(), TEXT("failed to find DLC plugin %s"), *Package->DLCName.ToString());
+
+
+	// ensure the plugin is mounted - will also load the shader code library if there is one
+	IPluginManager::Get().MountNewlyCreatedPlugin(Package->DLCName.ToString());
+
+	// load the asset registry for the new plugin
+	if (Package->DLCPlugin->CanContainContent())
+	{
+		FString AssetRegistryFileName = Package->DLCPlugin->GetBaseDir() / TEXT("AssetRegistry.bin");
+		AssetRegistryFileName.ReplaceInline(*Package->MountPath, TEXT("../../..") );
+
+		FAssetRegistryState PluginAssetRegistry;
+		if (FAssetRegistryState::LoadFromDisk(*AssetRegistryFileName, FAssetRegistryLoadOptions(), PluginAssetRegistry))
+		{
+			IAssetRegistry::GetChecked().AppendState(PluginAssetRegistry);
+		}
+		else
+		{
+			UE_LOG(LogGDKDLC, Error, TEXT("Failed to load plugin asset registry state %s"), *AssetRegistryFileName);
+		}
+	}
+
+	// all done
+	Package->MountState = EMountState::Mounted;
+	ProcessingState = EProcessingState::Completed;
+}
+
+
+void FGDKDLCHandler::Complete_GameThread( TSharedPtr<FPackageData> Package )
+{
+	check(IsInGameThread());
+
+	// signal callback, if any
+	if (Package->OnMounted != nullptr)
+	{
+		bool bSuccess = (Package->MountState == EMountState::Mounted);
+		Package->OnMounted(Package->DLCName, bSuccess);
+		Package->OnMounted = nullptr;
+	}
+
+	// remove the current work item & check for more work next tick
+	CurrentProcessingPackage.Reset();
+	ProcessingState = EProcessingState::Idle;
+}
+
+
+void FGDKDLCHandler::LicenseLost_GameThread( TSharedPtr<FPackageData> Package )
+{
+	FPlatformUserId OldStoreUserId = Package->StoreUserId;
+
+	// close the existing license
+	if (Package->DLCLicense != nullptr)
+	{
+		XStoreUnregisterPackageLicenseLost( Package->DLCLicense, Package->PackageLicenseLostToken, true);
+		XStoreCloseLicenseHandle(Package->DLCLicense);
+		Package->DLCLicense = nullptr;
+	}
+	Package->LicenseState = ELicenseState::LicenseLost;
+	Package->StoreUserId = PLATFORMUSERID_NONE;
+
+	// signal that the license has been lost. it's up to the title whether they want to retry/unmount/terminate etc.
+	OnDLCLicenseLost.Broadcast(Package->DLCName, OldStoreUserId);
+}
+
+
+
+
+void FGDKDLCHandler::UnmountDLCUnsafe( TSharedPtr<FPackageData> Package )
+{
+	FScopeLock Lock(&PackageDataLock);
+
+	if (Package->DLCHandle != nullptr)
+	{
+		// attempt to unregister the plugin as much as possible
+		FText FailureReason;
+		if (Package->DLCPlugin->GetDescriptor().bExplicitlyLoaded && IPluginManager::Get().UnmountExplicitlyLoadedPlugin(Package->DLCName.ToString(), &FailureReason))
+		{
+			const FString PluginFileName = FPaths::Combine( Package->DLCPlugin->GetBaseDir(), Package->DLCName.ToString() + TEXT(".uplugin") );
+			IPluginManager::Get().RemoveFromPluginsList( PluginFileName );
+		}
+		else
+		{
+			// not possible to properly unmount, so just simulate the results
+			IPluginManager::Get().OnPluginUnmounted().Broadcast(*Package->DLCPlugin);
+			FPackageName::UnRegisterMountPoint( Package->DLCPlugin->GetMountedAssetPath(), Package->DLCPlugin->GetContentDir() );
+		}
+
+		// unmount all pak files
+		TArray<FString> PakFileNames;
+		IPlatformFile::GetPlatformPhysical().FindFilesRecursively(PakFileNames, *Package->MountPath, TEXT(".pak") );
+		for (const FString& PakFileName : PakFileNames)
+		{
+			if (IGDKPackageManifestModule::Get().IsPakFileInstalled(Package->DLCHandle, PakFileName))
+			{
+				bool bUnmounted = FCoreDelegates::OnUnmountPak.IsBound() ? FCoreDelegates::OnUnmountPak.Execute(PakFileName) : false;
+				if (bUnmounted)
+				{
+					UE_LOG(LogGDKDLC, Log, TEXT("Successfully unmounted %s"), *PakFileName);
+				}
+				else
+				{
+					UE_LOG(LogGDKDLC, Log, TEXT("Failed to unmount %s"), *PakFileName);
+				}
+			}
+		}
+
+		// unmount the package
+		XPackageCloseMountHandle(Package->DLCHandle);
+		Package->DLCHandle = nullptr;
+
+		Package->MountState = EMountState::Unmounted;
+	}
+}
+
+
+void FGDKDLCHandler::AddToTicker()
+{
+	if (!TickHandle.IsValid())
+	{
+		FTSTicker& Ticker = FTSTicker::GetCoreTicker();
+
+		// Register delegate for ticker callback
+		FTickerDelegate TickDelegate = FTickerDelegate::CreateRaw(this, &FGDKDLCHandler::Tick);
+		TickHandle = Ticker.AddTicker(TickDelegate, 0.1f);
+	}
+}
+
+void FGDKDLCHandler::RemoveFromTicker()
+{
+	FTSTicker& Ticker = FTSTicker::GetCoreTicker();
+
+	// Unregister ticker delegate
+	if (TickHandle.IsValid())
+	{
+		Ticker.RemoveTicker(TickHandle);
+		TickHandle.Reset();
+	}
+}
+
+
+
+void FGDKDLCHandler::OnPackageLicenseLost(void* Context)
+{
+	TSharedPtr<FPackageData> Package = ((FPackageData::FContext*)Context)->Package.Pin();
+	{
+		FScopeLock Lock(&Package->Owner->PackageDataLock);
+		Package->Owner->LostLicensePackages.AddUnique(Package);
+		Package->Owner->AddToTicker();
+
+		XStoreUnregisterPackageLicenseLost( Package->DLCLicense, Package->PackageLicenseLostToken, false);
+	}
+}
+
+
+
+void FGDKDLCHandler::SetStoreUser( FPlatformUserId PlatformUserId )
+{
+	FGDKUserHandle NewStoreUser = IGDKRuntimeModule::Get().GetUserHandleByPlatformId(PlatformUserId);
+	if (NewStoreUser != CurrentStoreUser)
+	{
+		CurrentStoreUser = NewStoreUser;
+		CurrentStoreUserId = PlatformUserId;
+
+		if (CurrentStoreContext != nullptr)
+		{
+			XStoreCloseContextHandle(CurrentStoreContext);
+			CurrentStoreContext = nullptr;
+		}
+	}
+}
+
+bool FGDKDLCHandler::EnsureValidStoreContext()
+{
+	if (CurrentStoreUser.IsValid() && CurrentStoreContext == nullptr)
+	{
+		GDK_SCOPE_NOT_TIME_SENSITIVE(); // XStoreCreateContext is not safe to call on a time-sensitive thread
+
+		HRESULT hResult = XStoreCreateContext(CurrentStoreUser, &CurrentStoreContext);
+		if (FAILED(hResult))
+		{
+			UE_LOG(LogGDKDLC, Error, TEXT("Failed to create store context for user %s : 0x%X"), *LexToString(CurrentStoreUser), hResult);
+			return false;
+		}
+	}
+
+	return true;
+}
+
+
+
+FString FGDKDLCHandler::FindStoreId( FName DLCName )
+{
+	FScopeLock Lock(&PackageDataLock);
+
+	// see if the name is known
+	for (TSharedPtr<FPackageData> PackageData : AllPackages)
+	{
+		if (PackageData->DLCName == DLCName)
+		{
+			return UTF8_TO_TCHAR(PackageData->StoreId);
+		}
+	}
+
+	// name is not known - maybe the package isn't mounted yet. Find the store id for the given DLC
+	FString StoreId;
+
+	TArray<FString> DLCPlugins;
+	if (GConfig->GetArray(IGDKRuntimeModule::Get().GetConfigSectionName(), TEXT("DLCPlugins"), DLCPlugins, GEngineIni))
+	{
+		for (FString& DLCPlugin : DLCPlugins)
+		{
+			DLCPlugin.TrimStartAndEndInline();
+			DLCPlugin.ReplaceInline(TEXT("("), TEXT(""));
+			DLCPlugin.ReplaceInline(TEXT(")"), TEXT(""));
+		
+			FString DLCPluginName;
+			if (FParse::Value(*DLCPlugin, TEXT("DLCName="), DLCPluginName) && DLCPluginName == DLCName)
+			{
+				FParse::Value(*DLCPlugin, TEXT("StoreId="), StoreId);
+				break;
+			}
+		}
+	}
+
+	return StoreId;
+}
+
+
+
+TSharedPtr<FGDKDLCHandler::FPackageData> FGDKDLCHandler::FindDLC( FName DLCName )
+{
+	FScopeLock Lock(&PackageDataLock);
+
+	// see if the name is known
+	for (TSharedPtr<FPackageData> PackageData : AllPackages)
+	{
+		if (PackageData->DLCName == DLCName)
+		{
+			return PackageData;
+		}
+	}
+
+	// find the package with the store id
+	FString StoreId = FindStoreId(DLCName);
+	if (!StoreId.IsEmpty())
+	{
+		for (TSharedPtr<FPackageData> PackageData : AllPackages)
+		{
+			if (PackageData->DLCName == StoreId)
+			{
+				PackageData->DLCName = DLCName; // store the real DLC name for next time
+				return PackageData;
+			}
+		}
+	}
+
+	return nullptr;
+}
+
+bool FGDKDLCHandler::IsDLCAvailable( FName DLCName )
+{
+	TSharedPtr<FPackageData> Package = FindDLC(DLCName);
+	return Package.IsValid();
+}
+
+bool FGDKDLCHandler::IsDLCLicenseValid( FName DLCName )
+{
+	TSharedPtr<FPackageData> Package = FindDLC(DLCName);
+	return Package.IsValid() && (Package->LicenseState == ELicenseState::Licensed || Package->LicenseState == ELicenseState::Unlicensed_DevIgnore);
+}
+
+bool FGDKDLCHandler::IsDLCMounted( FName DLCName )
+{
+	TSharedPtr<FPackageData> Package = FindDLC(DLCName);
+	return Package.IsValid() && (Package->MountState == EMountState::Mounted);
+
+}
+bool FGDKDLCHandler::MountDLC( FName DLCName, FGDKDLCMountedCallback OnMounted )
+{
+	TSharedPtr<FPackageData> Package = FindDLC(DLCName);
+	if (!Package.IsValid())
+	{
+		return false;
+	}
+
+	Package->OnMounted = MoveTemp(OnMounted);
+	AddToPendingPackages(Package);
+
+	return true;
+}
+
+
+FString FGDKDLCHandler::GetDLCRoot( FName DLCName )
+{
+	TSharedPtr<FPackageData> Package = FindDLC(DLCName);
+	if (ensure(Package.IsValid() && Package->MountState == EMountState::Mounted))
+	{
+		return Package->MountPath;
+	}
+
+	return FString();
+}
+
+TArray<FName> FGDKDLCHandler::GetAllDLCNames()
+{
+	FScopeLock Lock(&PackageDataLock);
+
+	TArray<FName> DLCNames;
+	for (TSharedPtr<FPackageData> Package : AllPackages)
+	{
+		DLCNames.Add(Package->DLCName);
+	}
+
+	return MoveTemp(DLCNames);
+}
+
+
+
+FGDKDLCHandler::FPackageData::FPackageData(FGDKDLCHandler* InOwner, const XPackageDetails& InDetails)
+	: MountState(EMountState::NotMounted)
+	, LicenseState(ELicenseState::Unlicensed)
+	, Owner(InOwner)
+	, bIsAgeRestricted(InDetails.ageRestricted)
+	, DLCHandle(nullptr)
+	, DLCLicense(nullptr)
+	, StoreUserId(PLATFORMUSERID_NONE)
+{
+	FCStringAnsi::Strncpy(PackageIdentifier, InDetails.packageIdentifier, XPACKAGE_IDENTIFIER_MAX_LENGTH);
+	FCStringAnsi::Strncpy(StoreId, InDetails.storeId, STORE_SKU_ID_SIZE);
+
+	// set the name
+	PackageName = UTF8_TO_TCHAR(InDetails.displayName);
+	if (PackageName.IsEmpty())
+	{
+		PackageName = UTF8_TO_TCHAR(PackageIdentifier);
+	}
+
+	// set the default DLCName to the store id so it can be found even when the DLC isn't mounted
+	DLCName = UTF8_TO_TCHAR(InDetails.storeId);
+}
+
+
+
+FGDKDLCHandler::FPackageData::~FPackageData()
+{
+	if (DLCHandle != nullptr)
+	{
+		XPackageCloseMountHandle(DLCHandle);
+		DLCHandle = nullptr;
+	}
+
+	if (DLCLicense != nullptr)
+	{
+		XStoreUnregisterPackageLicenseLost( DLCLicense, PackageLicenseLostToken, true);
+		XStoreCloseLicenseHandle(DLCLicense);
+		DLCLicense = nullptr;
+	}
+}
+
+
+
+
+void FGDKDLCHandler::DebugUninstallDLC( FName DLCName )
+{
+	GDK_SCOPE_NOT_TIME_SENSITIVE(); // (debug only) XPackageUninstallPackage is not safe to call on a time-sensitive thread
+
+	for (int32 Idx = 0; Idx < AllPackages.Num(); Idx++)
+	{
+		TSharedPtr<FPackageData> Package = AllPackages[Idx];
+		if (Package->DLCName == DLCName)
+		{
+			UnmountDLCUnsafe(Package);
+			XPackageUninstallPackage(TCHAR_TO_UTF8(Package->PackageIdentifier));
+
+			AllPackages.RemoveAt(Idx);
+			break;
+		}
+	}
+}
+
+void FGDKDLCHandler::DebugUnmountDLC( FName DLCName )
+{
+	TSharedPtr<FGDKDLCHandler::FPackageData> Package = FindDLC(DLCName);
+	if (Package.IsValid())
+	{
+		UnmountDLCUnsafe(Package);
+	}
+}
+
+void FGDKDLCHandler::DebugDownloadDLC( FName DLCName )
+{
+	TSharedPtr<FGDKDLCHandler::FPackageData> Package = FindDLC(DLCName);
+	if (Package.IsValid())
+	{
+		// if we know about the package, it's already downloaded
+	}
+	else
+	{
+		// find the package with the store id
+		FString StoreId = FindStoreId(DLCName);
+		if (!StoreId.IsEmpty())
+		{
+			// make sure there is a store context
+			if (!EnsureValidStoreContext())
+			{
+				return;
+			}
+
+			// kick off the async 
+			HRESULT hResult = AsyncGDKTask( 
+				[this,StoreId]( XAsyncBlock* AsyncBlock )
+				{
+					char StoreIdAnsi[64];
+					FCStringAnsi::Strncpy(StoreIdAnsi, TCHAR_TO_ANSI(*StoreId), UE_ARRAY_COUNT(StoreIdAnsi));
+
+					const char* StoreIds[] =
+					{
+						StoreIdAnsi
+					};
+
+					return XStoreDownloadAndInstallPackagesAsync( CurrentStoreContext, StoreIds, 1, AsyncBlock );
+				},
+				[this,DLCName]( XAsyncBlock* AsyncBlock )
+				{
+					HRESULT hResult;
+
+					uint32 NumPackages = 0;
+					hResult = XStoreDownloadAndInstallPackagesResultCount( AsyncBlock, &NumPackages );
+					if (FAILED(hResult))
+					{
+						UE_LOG(LogGDKDLC, Error, TEXT("failed to get installing packages count for %s : 0x%X"), *DLCName.ToString(), hResult );
+						return;
+					}
+
+					auto PackageIdentifiers = new char[NumPackages][XPACKAGE_IDENTIFIER_MAX_LENGTH];
+					hResult = XStoreDownloadAndInstallPackagesResult( AsyncBlock, NumPackages, &PackageIdentifiers[0] );
+					if (FAILED(hResult))
+					{
+						delete[] PackageIdentifiers;
+						UE_LOG(LogGDKDLC, Error, TEXT("failed to get installing packages identifiers for %s : 0x%X"), *DLCName.ToString(), hResult );
+						return;
+					}
+
+					// the download was successfully started. we must have a valid license for this before we can mount it
+					for (uint32 Index = 0; Index < NumPackages; Index++)
+					{
+						UE_LOG(LogGDKDLC, Log, TEXT("Downloading package %s for DLC %s"), UTF8_TO_TCHAR(PackageIdentifiers[Index]), *DLCName.ToString() );
+					}
+	
+					delete[] PackageIdentifiers;
+				}
+			);
+		}
+	}
+}
+
+void FGDKDLCHandler::DebugDumpDLCState()
+{
+	UE_LOG( LogGDKDLC, Log, TEXT("DLC:") );
+	UE_LOG( LogGDKDLC, Log, TEXT("Processing State: %s"), *LexToString(ProcessingState) );
+	for (TSharedPtr<FPackageData> Package : AllPackages)
+	{
+		UE_LOG( LogGDKDLC, Log, TEXT("\t%-16s %-22s %-22s %s"), *Package->DLCName.ToString(), *LexToString(Package->MountState), *LexToString(Package->LicenseState), *Package->PackageName );
+	}
+}
+
+
+void FGDKDLCHandler::DumpDLCFiles( FName DLCName )
+{
+	for (TSharedPtr<FPackageData> Package : AllPackages)
+	{
+		if (Package->DLCName == DLCName && Package->MountState == EMountState::Mounted)
+		{
+			UE_LOG(LogGDKDLC, Log, TEXT("All Files For DLC '%s' in Mount Point '%s'"), *DLCName.ToString(), *Package->MountPath);
+
+			class FDumpFileVistor : public IPlatformFile::FDirectoryVisitor
+			{
+			public:
+				virtual bool Visit(const TCHAR* FilenameOrDirectory, bool bIsDirectory)
+				{
+					UE_CLOG(!bIsDirectory, LogGDKDLC, Log, TEXT("%s"), FilenameOrDirectory);
+					return true;
+				}
+			};
+			FDumpFileVistor Visitor;
+			IPlatformFile::GetPlatformPhysical().IterateDirectoryRecursively(*Package->MountPath, Visitor);
+			break;
+		}
+	}
+}
+
+
+FString FGDKDLCHandler::Debug_GetStateDescription() const
+{
+	return FString::Printf( TEXT("GDK DLC.  Store User: %-3d %-16s   Processing State: %s"), CurrentStoreUserId.GetInternalId(), *IGDKRuntimeModule::Get().GetGamertag(CurrentStoreUser), *LexToString(ProcessingState) );
+
+}
+
+bool FGDKDLCHandler::Debug_GetPackageDescription( int PackageIndex, FString& OutDescription, FColor& OutColor ) const
+{
+	TSharedPtr<FPackageData> Package;
+	{
+		FScopeLock Lock(&PackageDataLock);
+		if (!AllPackages.IsValidIndex(PackageIndex))
+		{
+			return false;
+		}
+		Package = AllPackages[PackageIndex];
+	}
+
+	if (Package->MountState == EMountState::NotMounted || Package->MountState == EMountState::Unmounted)
+	{
+		OutColor = FColor::Silver;
+	}
+	else if (Package->MountState == EMountState::Mounted)
+	{
+		bool bIsLicensed = (Package->LicenseState == ELicenseState::Licensed) || (Package->LicenseState == ELicenseState::Unlicensed_DevIgnore);
+		OutColor = bIsLicensed ? FColor::Green : FColor::Yellow;
+	}
+	else if (Package->MountState < EMountState::Mounted)
+	{
+		OutColor = FColor::Yellow;
+	}
+	else
+	{
+		OutColor = FColor::Red;
+	}
+
+	OutDescription = FString::Printf(TEXT(".......%-30s %-30s %-30s (%s)"), *Package->DLCName.ToString(), *LexToString(Package->MountState), *LexToString(Package->LicenseState), *Package->PackageName );
+	return true;
+}
+
+
+
+FString LexToString( FGDKDLCHandler::EMountState MountState )
+{
+	switch (MountState)
+	{
+		case FGDKDLCHandler::EMountState::NotMounted:              return TEXT("NotMounted");
+		case FGDKDLCHandler::EMountState::Mounted:                 return TEXT("Mounted");
+		case FGDKDLCHandler::EMountState::Unmounted:               return TEXT("Unmounted");
+		case FGDKDLCHandler::EMountState::Ignored_AgeRestricted:   return TEXT("Ignored(AgeRestricted)");
+		default:                                                   return TEXT("Invalid");
+	}
+}
+
+FString LexToString( FGDKDLCHandler::ELicenseState LicenseState )
+{
+	switch(LicenseState)
+	{
+		case FGDKDLCHandler::ELicenseState::Unlicensed:            return TEXT("Unlicensed");
+		case FGDKDLCHandler::ELicenseState::Licensed:              return TEXT("Licensed");
+		case FGDKDLCHandler::ELicenseState::LicenseLost:           return TEXT("LicenseLost");
+		case FGDKDLCHandler::ELicenseState::Unlicensed_DevIgnore:  return TEXT("Unlicensed(Sideload)");
+        default:                                                   return TEXT("Invalid");
+	}
+}
+
+FString LexToString( FGDKDLCHandler::EProcessingState ProcessingState )
+{
+	switch(ProcessingState)
+	{
+        case FGDKDLCHandler::EProcessingState::Idle:               return TEXT("Idle");
+		case FGDKDLCHandler::EProcessingState::QueryingLicense:    return TEXT("QueryingLicense");
+		case FGDKDLCHandler::EProcessingState::AcquiringLicense:   return TEXT("AcquiringLicense");
+		case FGDKDLCHandler::EProcessingState::MountingPackage:    return TEXT("MountingPackage");
+		case FGDKDLCHandler::EProcessingState::CachingData:        return TEXT("CachingData");
+		case FGDKDLCHandler::EProcessingState::MountingPlugin:     return TEXT("MountingPlugin");
+        case FGDKDLCHandler::EProcessingState::Completed:          return TEXT("Completed");
+		case FGDKDLCHandler::EProcessingState::Destroying:         return TEXT("Destroying");
+        default:                                                   return TEXT("Invalid");
+	}
+}
+
+#endif //WITH_GRDK
diff --git a/Engine/Plugins/Experimental/Microsoft/GDKDLCSupport/Source/GDKDLCSupport/Private/GDKDLCHandler.h b/Engine/Plugins/Experimental/Microsoft/GDKDLCSupport/Source/GDKDLCSupport/Private/GDKDLCHandler.h
new file mode 100644
--- /dev/null
+++ b/Engine/Plugins/Experimental/Microsoft/GDKDLCSupport/Source/GDKDLCSupport/Private/GDKDLCHandler.h
@@ -0,0 +1,152 @@
+// Copyright Epic Games, Inc. All Rights Reserved.
+
+#pragma once
+
+#include "GDKRuntimeModule.h"
+#if WITH_GRDK
+#include "GDKDLCSupportModule.h"
+#include "GDKTaskQueueHelpers.h"
+#include "Containers/Array.h"
+#include "Containers/Ticker.h"
+#include "Math/Color.h"
+
+#include "Microsoft/AllowMicrosoftPlatformTypes.h"
+THIRD_PARTY_INCLUDES_START
+#include <XPackage.h>
+#include <XStore.h>
+#include <atomic>
+THIRD_PARTY_INCLUDES_END
+#include "Microsoft/HideMicrosoftPlatformTypes.h"
+
+DECLARE_LOG_CATEGORY_EXTERN(LogGDKDLC, Log, All);
+
+class FGDKDLCHandler
+{
+public:
+	FGDKDLCHandler();
+	~FGDKDLCHandler();
+
+	// main public API support
+	bool IsDLCAvailable( FName DLCName );
+	bool IsDLCMounted( FName DLCName );
+	bool IsDLCLicenseValid( FName DLCName );
+	bool MountDLC( FName DLCName, FGDKDLCMountedCallback OnMounted = nullptr );
+	FString GetDLCRoot( FName DLCName );
+	TArray<FName> GetAllDLCNames();
+	void SetStoreUser( FPlatformUserId PlatformUserId );
+
+	FGDKDLCLicenseLostDelegate OnDLCLicenseLost;
+
+
+	// console command support
+	void DebugDumpDLCState();
+	void DebugUnmountDLC( FName DLCName );
+	void DebugUninstallDLC( FName DLCName );
+	void DebugDownloadDLC( FName DLCName );
+	void DumpDLCFiles( FName DLCName );
+
+	// support for debug rendering
+	FString Debug_GetStateDescription() const;
+	bool Debug_GetPackageDescription( int PackageIndex, FString& OutDescription, FColor& OutColor ) const;
+
+private:
+	enum class EMountState : uint8
+	{
+		NotMounted,
+		Mounted,
+		Unmounted,
+		Ignored_AgeRestricted,
+	};
+	enum class ELicenseState : uint8
+	{
+		Unlicensed,
+		Licensed,
+		LicenseLost,
+		Unlicensed_DevIgnore,
+	};
+
+	struct FPackageData
+	{
+		FPackageData(FGDKDLCHandler* InOwner, const XPackageDetails& InDetails);
+		~FPackageData();
+
+		char                        PackageIdentifier[XPACKAGE_IDENTIFIER_MAX_LENGTH];
+		char                        StoreId[STORE_SKU_ID_SIZE];
+		std::atomic<EMountState>    MountState;
+		std::atomic<ELicenseState>  LicenseState;
+		FString                     PackageName;
+		FString                     MountPath;
+		FGDKDLCHandler*             Owner;
+		bool                        bIsAgeRestricted;
+		FName                       DLCName;
+		XPackageMountHandle         DLCHandle;
+		XStoreLicenseHandle         DLCLicense;
+		XTaskQueueRegistrationToken PackageLicenseLostToken;
+		FPlatformUserId             StoreUserId;
+		TSharedPtr<class IPlugin>   DLCPlugin;
+		FGDKDLCMountedCallback      OnMounted;
+
+		// context used for GDK callbacks that only accept void*
+		struct FContext
+		{
+			TWeakPtr<FPackageData> Package;
+		};
+		FContext Context;
+	};
+	TArray<TSharedPtr<FPackageData>> AllPackages;
+	mutable FCriticalSection         PackageDataLock;
+	XTaskQueueRegistrationToken      PackageInstalledToken;
+
+	void AsyncInitDLC(const XPackageDetails& Details, bool bMount);
+	TSharedPtr<FPackageData> FindDLC( FName DLCName );
+
+	// DLC processing
+	enum class EProcessingState : uint8
+	{
+		Idle,
+		QueryingLicense,
+		AcquiringLicense,
+		MountingPackage,
+		CachingData,
+		MountingPlugin,
+		Completed,
+		Destroying,
+	};
+	std::atomic<EProcessingState>    ProcessingState = EProcessingState::Idle;
+	TArray<TSharedPtr<FPackageData>> PendingPackages;
+	TArray<TSharedPtr<FPackageData>> LostLicensePackages;
+	TSharedPtr<FPackageData>         CurrentProcessingPackage;
+	FGDKAsyncTaskMonitor             CurrentProcessingAsyncTask;
+
+	void AddToPendingPackages( TSharedPtr<FPackageData> Package );
+	void InitializePackage_GameThread( TSharedPtr<FPackageData> Package );
+	void AsyncQueryLicense_GameThread( TSharedPtr<FPackageData> Package );
+	void AsyncAcquireLicense_GameThread( TSharedPtr<FPackageData> Package );
+	void AsyncMountPackage_GameThread( TSharedPtr<FPackageData> Package );
+	void CachePackageData_BackgroundThread( TSharedPtr<FPackageData> Package );
+	void MountPlugin_GameThread( TSharedPtr<FPackageData> Package );
+	void Complete_GameThread( TSharedPtr<FPackageData> Package );
+	void LicenseLost_GameThread( TSharedPtr<FPackageData> Package );
+	void UnmountDLCUnsafe( TSharedPtr<FPackageData> Package );
+
+	// store
+	XStoreContextHandle CurrentStoreContext;
+	FGDKUserHandle      CurrentStoreUser;
+	FPlatformUserId     CurrentStoreUserId;
+	bool EnsureValidStoreContext();
+	FString FindStoreId( FName DLCName );
+	static void OnPackageLicenseLost(void* Context);
+
+	// ticker
+	FTSTicker::FDelegateHandle TickHandle;
+	void AddToTicker();
+	void RemoveFromTicker();
+	bool Tick( float DeltaSeconds );
+
+	// internal LexToString
+	friend FString LexToString( EMountState MountState );
+	friend FString LexToString( ELicenseState LicenseState );
+	friend FString LexToString( EProcessingState ProcessingState );
+};
+
+#endif //WITH_GRDK
diff --git a/Engine/Plugins/Experimental/Microsoft/GDKDLCSupport/Source/GDKDLCSupport/Private/GDKDLCSupportModule.cpp b/Engine/Plugins/Experimental/Microsoft/GDKDLCSupport/Source/GDKDLCSupport/Private/GDKDLCSupportModule.cpp
new file mode 100644
--- /dev/null
+++ b/Engine/Plugins/Experimental/Microsoft/GDKDLCSupport/Source/GDKDLCSupport/Private/GDKDLCSupportModule.cpp
@@ -0,0 +1,256 @@
+// Copyright Epic Games, Inc. All Rights Reserved.
+
+#include "GDKDLCSupportModule.h"
+#include "GDKDLCHandler.h"
+#include "GDKRuntimeModule.h"
+#include "Modules/ModuleInterface.h"
+#include "Modules/ModuleManager.h"
+#if WITH_GRDK
+#include "Templates/UniquePtr.h"
+#include "HAL/IConsoleManager.h"
+#include "Engine/Engine.h"
+#endif
+
+class FGDKDLCSupportModule : public IGDKDLCSupportModule
+{
+public:
+
+#if WITH_GRDK
+	virtual void StartupModule() override
+	{
+	}
+
+	virtual void ShutdownModule() override
+	{
+		SetOnScreenDebugRender(false);
+		DLCHandler.Reset();
+	}
+
+	void SetOnScreenDebugRender( bool bEnable )
+	{
+#if !(UE_BUILD_SHIPPING || UE_BUILD_TEST)
+		bool bIsEnabled = (DebugDrawTickHandle.IsValid());
+		if (bIsEnabled != bEnable)
+		{
+			if (bEnable)
+			{
+				EnsureInitialized();
+				DebugDrawTickHandle = FTSTicker::GetCoreTicker().AddTicker( TEXT("DLCDebugDraw"), 0.0f, [this](float) { AddDebugMessages(); return true; });
+			}
+			else
+			{
+				FTSTicker::GetCoreTicker().RemoveTicker(DebugDrawTickHandle);
+				DebugDrawTickHandle.Reset();
+			}
+		}
+#endif
+	}
+
+
+	// IGDKDLCSupportModule implementation
+
+	virtual bool IsDLCAvailable( FName DLCName ) const override
+	{
+		EnsureInitialized();
+		return DLCHandler->IsDLCAvailable(DLCName);
+	}
+
+	virtual bool IsDLCLicenseValid( FName DLCName ) const override
+	{
+		EnsureInitialized();
+		return DLCHandler->IsDLCLicenseValid( DLCName );
+	}
+
+	virtual bool IsDLCMounted( FName DLCName ) const override
+	{
+		EnsureInitialized();
+		return DLCHandler->IsDLCMounted(DLCName);
+	}
+
+	virtual bool MountDLC( FName DLCName, FGDKDLCMountedCallback OnMounted ) const override
+	{
+		EnsureInitialized();
+		return DLCHandler->MountDLC( DLCName, MoveTemp(OnMounted) );
+	}
+
+	virtual FString GetDLCRoot( FName DLCName ) const override
+	{
+		EnsureInitialized();
+		return DLCHandler->GetDLCRoot(DLCName);
+	}
+
+	virtual TArray<FName> GetAllDLCNames() const override
+	{
+		EnsureInitialized();
+		return DLCHandler->GetAllDLCNames();
+	}
+
+	virtual void SetStoreUserForDLC( FPlatformUserId PlatformUserId ) const override
+	{
+		EnsureInitialized();
+		DLCHandler->SetStoreUser(PlatformUserId);
+	}
+
+	virtual FGDKDLCLicenseLostDelegate& OnDLCLicenseLost() const override
+	{
+		EnsureInitialized();
+		return DLCHandler->OnDLCLicenseLost;
+	}
+
+	mutable TUniquePtr<FGDKDLCHandler> DLCHandler;
+
+private:
+
+	void EnsureInitialized() const
+	{
+		// only create GDK DLC handler if the GDK runtime is available
+		if (!DLCHandler.IsValid() && IGDKRuntimeModule::Get().IsAvailable())
+		{
+			DLCHandler.Reset(new FGDKDLCHandler());
+		}
+	}
+
+
+#if !(UE_BUILD_SHIPPING || UE_BUILD_TEST)
+	FTSTicker::FDelegateHandle DebugDrawTickHandle;
+	void AddDebugMessages()
+	{
+		GEngine->AddOnScreenDebugMessage(INDEX_NONE, 0.0f, FColor::White, DLCHandler->Debug_GetStateDescription(), false);
+
+		int PackageIndex = 0;
+		FString PackageDescription;
+		FColor PackageColor;
+		while (DLCHandler->Debug_GetPackageDescription(PackageIndex++, PackageDescription, PackageColor) )
+		{
+			GEngine->AddOnScreenDebugMessage(INDEX_NONE, 0.0f, PackageColor, PackageDescription, false);
+		}
+	}
+#endif
+
+#endif //WITH_GRDK
+};
+
+
+
+
+
+#if WITH_GRDK
+static FAutoConsoleCommand CVarGDKMountDLC(
+	TEXT("GDK.DLC.Mount"),
+	TEXT("[DLCName] Mount the given DLC"),
+	FConsoleCommandWithArgsDelegate::CreateLambda( [](const TArray<FString>& Args )
+	{
+		FGDKDLCSupportModule* Module = FModuleManager::GetModulePtr<FGDKDLCSupportModule>("GDKDLCSupport");
+		if (Module != nullptr && Module->DLCHandler.IsValid() && Args.Num() > 0)
+		{
+			Module->DLCHandler->MountDLC( FName(Args[0]), [](FName Name, bool bSuccess) { UE_LOG(LogTemp,Log,TEXT("DLC Mount: %s : %s"), *Name.ToString(), bSuccess?TEXT("success"):TEXT("fail") ); } );
+		}
+	})
+);
+
+
+static FAutoConsoleCommand CVarGDKDownloadDLC(
+	TEXT("GDK.DLC.DebugDownload"),
+	TEXT("[DLCName] Download the given DLC (debug only - unsafe)"),
+	FConsoleCommandWithArgsDelegate::CreateLambda( [](const TArray<FString>& Args )
+	{
+		FGDKDLCSupportModule* Module = FModuleManager::GetModulePtr<FGDKDLCSupportModule>("GDKDLCSupport");
+		if (Module != nullptr && Module->DLCHandler.IsValid() && Args.Num() > 0)
+		{
+			Module->DLCHandler->DebugDownloadDLC( FName(Args[0]) );
+		}
+	})
+);
+
+
+static FAutoConsoleCommand CVarGDKUninstallDLC(
+	TEXT("GDK.DLC.DebugUninstall"),
+	TEXT("[DLCName] Uninstall the given DLC (debug only - unsafe)"),
+	FConsoleCommandWithArgsDelegate::CreateLambda( [](const TArray<FString>& Args )
+	{
+		FGDKDLCSupportModule* Module = FModuleManager::GetModulePtr<FGDKDLCSupportModule>("GDKDLCSupport");
+		if (Module != nullptr && Module->DLCHandler.IsValid() && Args.Num() > 0)
+		{
+			Module->DLCHandler->DebugUninstallDLC( FName(Args[0]) );
+		}
+	})
+);
+
+static FAutoConsoleCommand CVarGDKUnmountDLC(
+	TEXT("GDK.DLC.DebugUnmount"),
+	TEXT("[DLCName] Unmount the given DLC (debug only - unsafe)"),
+	FConsoleCommandWithArgsDelegate::CreateLambda( [](const TArray<FString>& Args )
+	{
+		FGDKDLCSupportModule* Module = FModuleManager::GetModulePtr<FGDKDLCSupportModule>("GDKDLCSupport");
+		if (Module != nullptr && Module->DLCHandler.IsValid() && Args.Num() > 0)
+		{
+			Module->DLCHandler->DebugUnmountDLC( FName(Args[0]) );
+		}
+	})
+);
+
+
+static FAutoConsoleCommand CVarGDKDumpDLCFiles(
+	TEXT("GDK.DLC.Dir"),
+	TEXT("[DLCName] Enumerates all files in the given DLC"),
+	FConsoleCommandWithArgsDelegate::CreateLambda( [](const TArray<FString>& Args )
+	{
+		FGDKDLCSupportModule* Module = FModuleManager::GetModulePtr<FGDKDLCSupportModule>("GDKDLCSupport");
+		if (Module != nullptr && Module->DLCHandler.IsValid() && Args.Num() > 0)
+		{
+			Module->DLCHandler->DumpDLCFiles( FName(Args[0]) );
+		}
+	})
+);
+
+
+static FAutoConsoleCommand CVarGDKDumpDLC(
+	TEXT("GDK.DLC.DumpState"),
+	TEXT("Dumps the state of the DLC"),
+	FConsoleCommandDelegate::CreateLambda([]()
+	{
+		FGDKDLCSupportModule* Module = FModuleManager::GetModulePtr<FGDKDLCSupportModule>("GDKDLCSupport");
+		if (Module != nullptr && Module->DLCHandler.IsValid())
+		{
+			Module->DLCHandler->DebugDumpDLCState();
+		}
+	})
+);
+
+
+#if !(UE_BUILD_SHIPPING || UE_BUILD_TEST)
+static FAutoConsoleCommand CVarGDKDLCShow(
+	TEXT("GDK.DLC.DebugShow"),
+	TEXT("Show or hide the on-screen debugging"),
+	FConsoleCommandWithArgsDelegate::CreateLambda( [](const TArray<FString>& Args )
+	{
+		FGDKDLCSupportModule* Module = FModuleManager::GetModulePtr<FGDKDLCSupportModule>("GDKDLCSupport");
+		if (Module != nullptr && Module->DLCHandler.IsValid())
+		{
+			bool bShow = (Args.Num() == 0) || (FCString::Atoi(*Args[0]) != 0);
+			Module->SetOnScreenDebugRender(bShow);
+		}
+	})
+);
+#endif
+
+
+static FAutoConsoleCommand CVarGDKDLCSetStoreUser(
+	TEXT("GDK.DLC.SetStoreUser"),
+	TEXT("[PlatformUserId] Change the store user for DLC licensing"),
+	FConsoleCommandWithArgsDelegate::CreateLambda( [](const TArray<FString>& Args )
+	{
+		FGDKDLCSupportModule* Module = FModuleManager::GetModulePtr<FGDKDLCSupportModule>("GDKDLCSupport");
+		if (Module != nullptr && Module->DLCHandler.IsValid())
+		{
+			Module->DLCHandler->SetStoreUser( FPlatformUserId::CreateFromInternalId( FCString::Atoi(*Args[0]) ));
+		}
+	})
+);
+
+
+#endif //WITH_GRDK
+
+
+
+IMPLEMENT_MODULE(FGDKDLCSupportModule, GDKDLCSupport);
diff --git a/Engine/Plugins/Experimental/Microsoft/GDKDLCSupport/Source/GDKDLCSupport/Public/GDKDLCSupportModule.h b/Engine/Plugins/Experimental/Microsoft/GDKDLCSupport/Source/GDKDLCSupport/Public/GDKDLCSupportModule.h
new file mode 100644
--- /dev/null
+++ b/Engine/Plugins/Experimental/Microsoft/GDKDLCSupport/Source/GDKDLCSupport/Public/GDKDLCSupportModule.h
@@ -0,0 +1,58 @@
+// Copyright Epic Games, Inc. All Rights Reserved.
+
+#pragma once
+
+#include "CoreMinimal.h"
+#include "Modules/ModuleInterface.h"
+#include "Modules/ModuleManager.h"
+#include "Delegates/Delegate.h"
+#include "GDKRuntimeModule.h"
+
+typedef TFunction<void(FName, bool)> FGDKDLCMountedCallback;
+typedef TMulticastDelegate<void(FName, FPlatformUserId )> FGDKDLCLicenseLostDelegate;
+
+class GDKDLCSUPPORT_API IGDKDLCSupportModule : public IModuleInterface
+{
+public:
+	/**
+	 * Singleton-like access to the MS Gaming Runtime module instance
+	 * @return Returns IGDKDLCSupportModule singleton instance, loading the module on demand if needed
+	 */
+	static inline IGDKDLCSupportModule& Get()
+	{
+		return FModuleManager::LoadModuleChecked<IGDKDLCSupportModule>("GDKDLCSupport");
+	}
+
+#if WITH_GRDK
+
+	/** NOTE: GDK DLC support is experimental and the API is subject to change **/
+
+	// determines whether the given DLC is available
+	virtual bool IsDLCAvailable( FName DLCName ) const = 0;
+
+	// query if the given DLC license is valid
+	virtual bool IsDLCLicenseValid( FName DLCName ) const = 0;
+
+	// determines whether the given DLC is mounted
+	virtual bool IsDLCMounted( FName DLCName ) const = 0;
+
+	// start asyncronsouly mounting the given DLC (not necessary with GDK.DLC.MountOnStartup / GDK.DLC.MountAtRuntime)
+	// can also be used to reaquire DLC license after it has been lost
+	virtual bool MountDLC( FName DLCName, FGDKDLCMountedCallback OnMounted = nullptr ) const = 0;
+
+	// returns the file system mount point for the given DLC
+	virtual FString GetDLCRoot( FName DLCName ) const = 0;
+
+	// returns the list of all known DLC
+	virtual TArray<FName> GetAllDLCNames() const = 0;
+
+	// when using GDK.DLC.AutoAcquireLicense, specifies the user who's credentials will be used
+	// note that this user may not actually need a license for the DLC as long as there is *a* user who's signed into the console has a license (they don't even need to be signed into the game)
+	virtual void SetStoreUserForDLC( FPlatformUserId PlatformUserId ) const = 0;
+
+	// called when DLCs license is lost. this typically happens when the last user who has a license to the DLC signs out (does not apply to sideloaded DLC)
+	virtual FGDKDLCLicenseLostDelegate& OnDLCLicenseLost() const = 0;
+
+#endif
+};
+
diff --git a/Engine/Plugins/Runtime/Microsoft/GDKPackageChunkInstall/GDKPackageChunkInstall.uplugin b/Engine/Plugins/Runtime/Microsoft/GDKPackageChunkInstall/GDKPackageChunkInstall.uplugin
new file mode 100644
--- /dev/null
+++ b/Engine/Plugins/Runtime/Microsoft/GDKPackageChunkInstall/GDKPackageChunkInstall.uplugin
@@ -0,0 +1,41 @@
+{
+	"FileVersion": 3,
+	"Version": 1,
+	"VersionName": "0.1",
+	"FriendlyName": "GDK Package Chunk Installer",
+	"Description": "Allows titles to use GDK chunk installer (XPackage)",
+	"CreatedBy": "Epic Games, Inc.",
+	"Category": "Platform",
+	"CreatedByURL": "https://epicgames.com",
+	"DocsURL": "",
+	"MarketplaceURL": "",
+	"SupportURL": "",
+	"EnabledByDefault": false,
+	"RequiresBuildPlatform": false,
+	"CanContainContent": false,
+	"IsBetaVersion": true,
+	"Installed": false,
+	"SupportedTargetPlatforms": [ "Win64" ],
+	"Modules": [
+		{
+			"Name": "GDKPackageChunkInstall",
+			"Type": "RuntimeNoCommandlet",
+			"LoadingPhase": "PostConfigInit",
+			"PlatformAllowList": [ "Win64" ],
+			"PlatformArchitectureDenyList": [ "Win64:arm64" ]
+		}
+	],
+	"Plugins":
+	[
+		{
+			"Name": "MSGameStore",
+			"Enabled": true,
+			"PlatformAllowList": [ "Win64" ]
+		},
+		{
+			"Name": "MSGamingRuntime",
+			"Enabled": true,
+			"PlatformAllowList": [ "Win64" ]
+		}
+	]
+}
\ No newline at end of file
diff --git a/Engine/Plugins/Runtime/Microsoft/GDKPackageChunkInstall/Source/GDKPackageChunkInstall/GDKPackageChunkInstall.Build.cs b/Engine/Plugins/Runtime/Microsoft/GDKPackageChunkInstall/Source/GDKPackageChunkInstall/GDKPackageChunkInstall.Build.cs
new file mode 100644
--- /dev/null
+++ b/Engine/Plugins/Runtime/Microsoft/GDKPackageChunkInstall/Source/GDKPackageChunkInstall/GDKPackageChunkInstall.Build.cs
@@ -0,0 +1,24 @@
+﻿// Copyright Epic Games, Inc. All Rights Reserved.
+
+using UnrealBuildTool;
+
+public class GDKPackageChunkInstall : ModuleRules
+{
+    public GDKPackageChunkInstall(ReadOnlyTargetRules Target) : base(Target)
+    {
+		if (Target.Platform == UnrealTargetPlatform.Win64)
+		{
+			BinariesSubFolder = "GDK";
+		}
+
+        PrivateDependencyModuleNames.AddRange(
+            new string[] 
+            {
+                "Core",
+				"GRDK",
+                "GDKPackageManifest",
+				"GDKRuntime",
+            }
+        );
+    }
+}
\ No newline at end of file
diff --git a/Engine/Plugins/Runtime/Microsoft/GDKPackageChunkInstall/Source/GDKPackageChunkInstall/Private/GDKPackageChunkInstallModule.cpp b/Engine/Plugins/Runtime/Microsoft/GDKPackageChunkInstall/Source/GDKPackageChunkInstall/Private/GDKPackageChunkInstallModule.cpp
new file mode 100644
--- /dev/null
+++ b/Engine/Plugins/Runtime/Microsoft/GDKPackageChunkInstall/Source/GDKPackageChunkInstall/Private/GDKPackageChunkInstallModule.cpp
@@ -0,0 +1,384 @@
+// Copyright Epic Games, Inc. All Rights Reserved.
+
+#include "GenericPlatform/GenericPlatformChunkInstall.h"
+#include "Modules/ModuleManager.h"
+#include "GDKPlatformChunkInstall.h"
+#include "IGDKPackageManifestModule.h"
+#include "GDKThreadCheck.h"
+#include "GDKTaskQueueHelpers.h"
+#include "GDKRuntimeModule.h"
+#include "HAL/IConsoleManager.h"
+
+#if WITH_GRDK
+#include "Microsoft/AllowMicrosoftPlatformTypes.h"
+THIRD_PARTY_INCLUDES_START
+#include <XPackage.h>
+#include <XGameErr.h>
+THIRD_PARTY_INCLUDES_END
+#include "Microsoft/HideMicrosoftPlatformTypes.h"
+
+// not relevent for general use
+#define WITH_XPACKAGE_DEBUG_CMDS 0
+
+
+/**
+ * Module for the GDK Package Chunk Installer
+ */
+class FGDKPackageChunkInstallModule : public IPlatformChunkInstallModule
+{
+public:
+	std::atomic<bool> bWantSystemCreation;
+	TUniquePtr<FGDKPlatformChunkInstall> ChunkInstaller;
+
+	FGDKPackageChunkInstallModule()
+		: bWantSystemCreation(true)
+	{
+#if WITH_EDITOR
+		// tear down the chunk installer when PIE finishes
+		IGDKRuntimeModule::Get().GetOnTeardownForPIE().AddRaw( this, &FGDKPackageChunkInstallModule::OnTeardownForPIE );
+#endif
+
+		// ensure the package manifest module is loaded if we'll need it asyncronously later on
+#if WITH_CHUNKINSTALL_ASYNC_INIT
+		IGDKPackageManifestModule::Get();
+#endif
+	}
+
+	virtual void ShutdownModule() override
+	{
+		ChunkInstaller.Reset();
+	
+#if WITH_EDITOR
+		if (IGDKRuntimeModule* GDKRuntime = IGDKRuntimeModule::TryGet())
+		{
+			GDKRuntime->GetOnTeardownForPIE().RemoveAll(this);
+		}
+#endif
+	}
+
+	virtual IPlatformChunkInstall* GetPlatformChunkInstall()
+	{
+		// prevent creation of GDK chunk installer if the GDK runtime is unavailable
+		if (!IGDKRuntimeModule::Get().IsAvailable())
+		{
+			bWantSystemCreation = false;
+		}
+
+		if (bWantSystemCreation)
+		{
+			bWantSystemCreation = false;
+			ChunkInstaller.Reset(new FGDKPlatformChunkInstall());
+		}
+
+		return ChunkInstaller.Get();
+	}
+
+#if WITH_EDITOR
+	void OnTeardownForPIE()
+	{
+		ChunkInstaller.Reset();
+		bWantSystemCreation = true;
+	}
+#endif
+
+
+#if WITH_XPACKAGE_DEBUG_CMDS
+	static void CALLBACK DebugInstallXPackageCallback(void*, XPackageInstallationMonitorHandle Handle)
+	{
+		XPackageInstallationProgress Progress;
+		XPackageGetInstallationProgress(Handle, &Progress);
+		UE_LOG(LogChunkInstaller, Log, TEXT("XPackage debug chunk installing %llu / %llu"), Progress.installedBytes, Progress.totalBytes );
+		if (Progress.completed)
+		{
+			UE_LOG(LogChunkInstaller, Log, TEXT("XPackage debug chunk installed") );
+			XPackageCloseInstallationMonitorHandle(Handle);
+		}
+	}
+
+	void DebugInstallXPackageChunk( XPackageChunkSelector& Selector )
+	{
+		GDK_SCOPE_NOT_TIME_SENSITIVE(); // (debug only) XPackageInstallChunks is not safe to call on a time-sensitive thread
+
+		char PackageIdentifier[XPACKAGE_IDENTIFIER_MAX_LENGTH];
+		XPackageGetCurrentProcessPackageIdentifier(XPACKAGE_IDENTIFIER_MAX_LENGTH, PackageIdentifier);
+
+		XPackageInstallationMonitorHandle Handle = nullptr;
+		HRESULT hr = XPackageInstallChunks(PackageIdentifier, 1, &Selector, 100, true, FGDKAsyncTaskQueue::GetGenericQueue(), &Handle);
+		if (SUCCEEDED(hr))
+		{
+			XTaskQueueRegistrationToken Token;
+			hr = XPackageRegisterInstallationProgressChanged( Handle, nullptr, DebugInstallXPackageCallback, &Token);
+			if (FAILED(hr))
+			{
+				XPackageCloseInstallationMonitorHandle(Handle);
+			}
+		}
+		UE_CLOG(FAILED(hr), LogChunkInstaller, Warning, TEXT("XPackageInstallChunks failed : 0x%X"), hr );
+	}
+
+	void DebugUninstallXPackageChunk( XPackageChunkSelector& Selector )
+	{
+		char PackageIdentifier[XPACKAGE_IDENTIFIER_MAX_LENGTH];
+		XPackageGetCurrentProcessPackageIdentifier(XPACKAGE_IDENTIFIER_MAX_LENGTH, PackageIdentifier);
+		HRESULT hr = XPackageUninstallChunks(PackageIdentifier, 1, &Selector );
+		UE_CLOG(FAILED(hr), LogChunkInstaller, Warning, TEXT("XPackageUninstallChunks failed : 0x%X"), hr );
+	}
+#endif //WITH_XPACKAGE_DEBUG_CMDS
+
+};
+
+#if !UE_BUILD_SHIPPING
+static FAutoConsoleCommand CVarChunkInstallDumpChunks(
+	TEXT("GDK.ChunkInstall.DumpChunks"),
+	TEXT("Dumps the state of the chunk installation"),
+	FConsoleCommandDelegate::CreateLambda([]()
+	{
+		FGDKPackageChunkInstallModule* Module = FModuleManager::GetModulePtr<FGDKPackageChunkInstallModule>("GDKPackageChunkInstall");
+		if (Module != nullptr && Module->ChunkInstaller.IsValid())
+		{
+			Module->ChunkInstaller->DebugDumpChunkState();
+		}
+	})
+);
+
+
+static FAutoConsoleCommand CVarChunkInstallDebugInstall(
+	TEXT("GDK.ChunkInstall.InstallNamedChunk"),
+	TEXT("[Name] Request installation of the given named chunk"),
+	FConsoleCommandWithArgsDelegate::CreateLambda([](const TArray<FString>& Args)
+	{
+		FGDKPackageChunkInstallModule* Module = FModuleManager::GetModulePtr<FGDKPackageChunkInstallModule>("GDKPackageChunkInstall");
+		if (Module != nullptr && Module->ChunkInstaller.IsValid() && Args.Num() > 0)
+		{
+			TArray<FName> NamedChunks; 
+			for (const FString& Arg : Args) NamedChunks.Add(FName(*Arg));
+			Module->ChunkInstaller->InstallNamedChunks(NamedChunks);
+		}
+	})
+);
+
+
+static FAutoConsoleCommand CVarChunkUninstallDebugInstall(
+	TEXT("GDK.ChunkInstall.UninstallNamedChunk"),
+	TEXT("[Name] Request removal of the given named chunk"),
+	FConsoleCommandWithArgsDelegate::CreateLambda([](const TArray<FString>& Args)
+	{
+		FGDKPackageChunkInstallModule* Module = FModuleManager::GetModulePtr<FGDKPackageChunkInstallModule>("GDKPackageChunkInstall");
+		if (Module != nullptr && Module->ChunkInstaller.IsValid() && Args.Num() > 0)
+		{
+			TArray<FName> NamedChunks; 
+			for (const FString& Arg : Args) NamedChunks.Add(FName(*Arg));
+			Module->ChunkInstaller->UninstallNamedChunks(NamedChunks);
+		}
+	})
+);
+
+
+static FAutoConsoleCommand GDKPackageDump(
+	TEXT("GDK.XPackage.Dump"),
+	TEXT("Dumps the details of all GDK XPackages"),
+	FConsoleCommandDelegate::CreateLambda([]()
+	{
+		GDK_SCOPE_NOT_TIME_SENSITIVE(); // (debug only) XPackageEnumeratePackages & XPackageEnumerateChunkAvailability are not safe to call on a time-sensitive thread
+
+		//early out if we're not a chunk build
+		if (!XPackageIsPackagedProcess())
+		{
+			UE_LOG(LogChunkInstaller, Log, TEXT("Not a packaged build") );
+			return;
+		}
+
+		auto EnumPackageCallback = [](void* Context, const XPackageDetails* Details)
+		{
+			//dump package details
+			UE_LOG(LogChunkInstaller, Log, TEXT("-------------") );
+			UE_LOG(LogChunkInstaller, Log, TEXT("Found %s package:"), Details->kind == XPackageKind::Content ? TEXT("Content") : TEXT("Game"));
+			UE_LOG(LogChunkInstaller, Log, TEXT("\tIdentifier:   %s"), *BytesToHex((uint8*)Details->packageIdentifier,XPACKAGE_IDENTIFIER_MAX_LENGTH) );
+			UE_LOG(LogChunkInstaller, Log, TEXT("\tVersion:      %d.%d.%d.%d"), Details->version.major, Details->version.minor, Details->version.build, Details->version.revision );
+			UE_LOG(LogChunkInstaller, Log, TEXT("\tDisplay Name: %s"), Details->displayName ? UTF8_TO_TCHAR(Details->displayName) : TEXT("<none>") );
+			UE_LOG(LogChunkInstaller, Log, TEXT("\tDescription:  %s"), Details->description ? UTF8_TO_TCHAR(Details->description) : TEXT("<none>") );
+			UE_LOG(LogChunkInstaller, Log, TEXT("\tPublisher:    %s"), Details->publisher   ? UTF8_TO_TCHAR(Details->publisher)   : TEXT("<none>") );
+			UE_LOG(LogChunkInstaller, Log, TEXT("\tStoreId:      %s"), Details->storeId     ? UTF8_TO_TCHAR(Details->storeId)     : TEXT("<none>") );
+			UE_LOG(LogChunkInstaller, Log, TEXT("\tInstalling:   %s"), *LexToString(Details->installing) );
+
+			// dump out the status of all known chunks in the package
+			UE_LOG(LogChunkInstaller, Log, TEXT("\tChunks:") );
+			auto EnumChunkCallback = [](void* Context, const XPackageChunkSelector* Selector, XPackageChunkAvailability Availability)
+			{
+				UE_LOG(LogChunkInstaller, Log, TEXT("\t\tChunk %-32s %s"), *LexToString(*Selector), *LexToString(Availability) );
+				return true;
+			};
+			XPackageEnumerateChunkAvailability( Details->packageIdentifier, XPackageChunkSelectorType::Chunk,    Context, EnumChunkCallback );
+			XPackageEnumerateChunkAvailability( Details->packageIdentifier, XPackageChunkSelectorType::Language, Context, EnumChunkCallback );
+			XPackageEnumerateChunkAvailability( Details->packageIdentifier, XPackageChunkSelectorType::Tag,      Context, EnumChunkCallback );
+			XPackageEnumerateChunkAvailability( Details->packageIdentifier, XPackageChunkSelectorType::Feature,  Context, EnumChunkCallback );
+			UE_LOG(LogChunkInstaller, Log, TEXT("-------------") );
+			return true;
+		};
+
+		// dump out the status of all known packages
+		UE_LOG(LogChunkInstaller, Log, TEXT("All XPackages:") );
+		if (XPackageEnumeratePackages( XPackageKind::Game, XPackageEnumerationScope::ThisAndRelated, nullptr, EnumPackageCallback ) == E_GAMEPACKAGE_NO_STORE_ID)
+		{
+			// this title is not configured in the store yet, so just enumerate the current package
+			UE_LOG(LogChunkInstaller, Log, TEXT("(title has no store Id - just showing this package") );
+			XPackageEnumeratePackages( XPackageKind::Game, XPackageEnumerationScope::ThisOnly, nullptr, EnumPackageCallback );
+		}
+		else
+		{
+			// enumerate other related content too, although it's not supported yet
+			XPackageEnumeratePackages( XPackageKind::Content, XPackageEnumerationScope::ThisAndRelated, nullptr, EnumPackageCallback );
+		}
+	})
+);
+#endif //!UE_BUILD_SHIPPING
+
+#if WITH_XPACKAGE_DEBUG_CMDS
+static FAutoConsoleCommand GDKPackageInstallLanguage(
+	TEXT("GDK.XPackage.InstallLanguage"),
+	TEXT(" [lang] Install the given language chunk"),
+	FConsoleCommandWithArgsDelegate::CreateLambda([](const TArray<FString>& Args)
+	{
+		FGDKPackageChunkInstallModule* Module = FModuleManager::LoadModulePtr<FGDKPackageChunkInstallModule>("GDKPackageChunkInstall");
+		if (Module != nullptr && Args.Num() == 1)
+		{
+			FTCHARToUTF8 UTF8Arg(*Args[0]);
+			XPackageChunkSelector Selector;
+			Selector.type = XPackageChunkSelectorType::Language;
+			Selector.language = UTF8Arg.Get();
+			Module->DebugInstallXPackageChunk(Selector);
+		}
+	})
+);
+
+static FAutoConsoleCommand GDKPackageInstallTag(
+	TEXT("GDK.XPackage.InstallTag"),
+	TEXT(" [tag] Install the given tagged chunk"),
+	FConsoleCommandWithArgsDelegate::CreateLambda([](const TArray<FString>& Args)
+	{
+		FGDKPackageChunkInstallModule* Module = FModuleManager::LoadModulePtr<FGDKPackageChunkInstallModule>("GDKPackageChunkInstall");
+		if (Module != nullptr && Args.Num() == 1)
+		{
+			FTCHARToUTF8 UTF8Arg(*Args[0]);
+			XPackageChunkSelector Selector;
+			Selector.type = XPackageChunkSelectorType::Tag;
+			Selector.tag = UTF8Arg.Get();
+			Module->DebugInstallXPackageChunk(Selector);
+		}
+	})
+);
+
+static FAutoConsoleCommand GDKPackageInstallChunk(
+	TEXT("GDK.XPackage.InstallChunk"),
+	TEXT("[id] Install the given chunk id"),
+	FConsoleCommandWithArgsDelegate::CreateLambda([](const TArray<FString>& Args)
+	{
+		FGDKPackageChunkInstallModule* Module = FModuleManager::LoadModulePtr<FGDKPackageChunkInstallModule>("GDKPackageChunkInstall");
+		if (Module != nullptr && Args.Num() == 1)
+		{
+			XPackageChunkSelector Selector;
+			Selector.type = XPackageChunkSelectorType::Chunk;
+			Selector.chunkId = FCString::Atoi(*Args[0]);
+			Module->DebugInstallXPackageChunk(Selector);
+		}
+	})
+);
+
+static FAutoConsoleCommand GDKPackageInstallFeature(
+	TEXT("GDK.XPackage.InstallFeature"),
+	TEXT(" [feature] Install the given feature"),
+	FConsoleCommandWithArgsDelegate::CreateLambda([](const TArray<FString>& Args)
+	{
+		FGDKPackageChunkInstallModule* Module = FModuleManager::LoadModulePtr<FGDKPackageChunkInstallModule>("GDKPackageChunkInstall");
+		if (Module != nullptr && Args.Num() == 1)
+		{
+			FTCHARToUTF8 UTF8Arg(*Args[0]);
+			XPackageChunkSelector Selector;
+			Selector.type = XPackageChunkSelectorType::Feature;
+			Selector.feature = UTF8Arg.Get();
+			Module->DebugInstallXPackageChunk(Selector);
+		}
+	})
+);
+
+static FAutoConsoleCommand GDKPackageUninstallLanguage(
+	TEXT("GDK.XPackage.UninstallLanguage"),
+	TEXT(" [lang] Uninstall the given language chunk"),
+	FConsoleCommandWithArgsDelegate::CreateLambda([](const TArray<FString>& Args)
+	{
+		FGDKPackageChunkInstallModule* Module = FModuleManager::LoadModulePtr<FGDKPackageChunkInstallModule>("GDKPackageChunkInstall");
+		if (Module != nullptr && Args.Num() == 1)
+		{
+			FTCHARToUTF8 UTF8Arg(*Args[0]);
+			XPackageChunkSelector Selector;
+			Selector.type = XPackageChunkSelectorType::Language;
+			Selector.language = UTF8Arg.Get();
+			Module->DebugUninstallXPackageChunk(Selector);
+		}
+	})
+);
+
+static FAutoConsoleCommand GDKPackageUninstallTag(
+	TEXT("GDK.XPackage.UninstallTag"),
+	TEXT(" [tag] Uninstall the given tagged chunk"),
+	FConsoleCommandWithArgsDelegate::CreateLambda([](const TArray<FString>& Args)
+	{
+		FGDKPackageChunkInstallModule* Module = FModuleManager::LoadModulePtr<FGDKPackageChunkInstallModule>("GDKPackageChunkInstall");
+		if (Module != nullptr && Args.Num() == 1)
+		{
+			FTCHARToUTF8 UTF8Arg(*Args[0]);
+			XPackageChunkSelector Selector;
+			Selector.type = XPackageChunkSelectorType::Tag;
+			Selector.tag = UTF8Arg.Get();
+			Module->DebugUninstallXPackageChunk(Selector);
+		}
+	})
+);
+
+static FAutoConsoleCommand GDKPackageUninstallChunk(
+	TEXT("GDK.XPackage.UninstallChunk"),
+	TEXT("[id] Uninstall the given chunk id"),
+	FConsoleCommandWithArgsDelegate::CreateLambda([](const TArray<FString>& Args)
+	{
+		FGDKPackageChunkInstallModule* Module = FModuleManager::LoadModulePtr<FGDKPackageChunkInstallModule>("GDKPackageChunkInstall");
+		if (Module != nullptr && Args.Num() == 1)
+		{
+			XPackageChunkSelector Selector;
+			Selector.type = XPackageChunkSelectorType::Chunk;
+			Selector.chunkId = FCString::Atoi(*Args[0]);
+			Module->DebugUninstallXPackageChunk(Selector);
+		}
+	})
+);
+
+static FAutoConsoleCommand GDKPackageUninstallFeature(
+	TEXT("GDK.XPackage.UninstallFeature"),
+	TEXT(" [feature] Uninstall the given feature"),
+	FConsoleCommandWithArgsDelegate::CreateLambda([](const TArray<FString>& Args)
+	{
+		FGDKPackageChunkInstallModule* Module = FModuleManager::LoadModulePtr<FGDKPackageChunkInstallModule>("GDKPackageChunkInstall");
+		if (Module != nullptr && Args.Num() == 1)
+		{
+			FTCHARToUTF8 UTF8Arg(*Args[0]);
+			XPackageChunkSelector Selector;
+			Selector.type = XPackageChunkSelectorType::Feature;
+			Selector.feature = UTF8Arg.Get();
+			Module->DebugUninstallXPackageChunk(Selector);
+		}
+	})
+);
+#endif //WITH_XPACKAGE_DEBUG_CMDS
+
+#else // ... WITH_GRDK
+
+// stub
+class FGDKPackageChunkInstallModule : public IPlatformChunkInstallModule
+{
+	virtual IPlatformChunkInstall* GetPlatformChunkInstall() override { return nullptr; }
+};
+#endif //WITH_GRDK
+
+
+
+IMPLEMENT_MODULE(FGDKPackageChunkInstallModule, GDKPackageChunkInstall);
+
diff --git a/Engine/Plugins/Runtime/Microsoft/GDKPackageChunkInstall/Source/GDKPackageChunkInstall/Private/GDKPlatformChunkInstall.cpp b/Engine/Plugins/Runtime/Microsoft/GDKPackageChunkInstall/Source/GDKPackageChunkInstall/Private/GDKPlatformChunkInstall.cpp
new file mode 100644
--- /dev/null
+++ b/Engine/Plugins/Runtime/Microsoft/GDKPackageChunkInstall/Source/GDKPackageChunkInstall/Private/GDKPlatformChunkInstall.cpp
@@ -0,0 +1,1443 @@
+// Copyright Epic Games, Inc. All Rights Reserved.
+
+#include "GDKPlatformChunkInstall.h"
+#if WITH_GRDK
+#include "GDKTaskQueueHelpers.h"
+#include "GDKThreadCheck.h"
+#include "IGDKPackageManifestModule.h"
+#include "HAL/FileManager.h"
+#include "Misc/ConfigCacheIni.h"
+#include "Misc/CoreDelegates.h"
+#include "Misc/Paths.h"
+#include "Misc/ScopeLock.h"
+#include "HAL/IConsoleManager.h"
+#include "Async/TaskGraphInterfaces.h"
+#include "Async/Async.h"
+#include "Tasks/Task.h"
+
+// default update interval for chunk installation in milliseconds
+#if !defined(GDK_CHUNKINSTALL_UPDATEINTERVAL_MS)
+	#define GDK_CHUNKINSTALL_UPDATEINTERVAL_MS 1000
+#endif
+
+
+static int32 GGDKChunkInstallSuppressUserConfirmation = 1;
+static FAutoConsoleVariableRef CVarChunkInstallSuppressUserConfirmation(
+	TEXT( "GDK.ChunkInstall.SuppressUserConfirmation" ),
+	GGDKChunkInstallSuppressUserConfirmation,
+	TEXT( "Whether to show a user confirmation when trying to download chunks over a certain size" ));
+
+static bool GGDKChunkInstallTerminateGameOnExternalFeatureRemoval = true;
+static FAutoConsoleVariableRef CVarChunkInstallTerminateGameOnExternalFeatureRemoval(
+	TEXT("GDK.ChunkInstall.TerminateOnExternalFeatureRemoval"),
+	GGDKChunkInstallTerminateGameOnExternalFeatureRemoval,
+	TEXT("Whether to terminate the game if the user removes a Feature via the 'Manage Game & Add-Ins' UI")); // note: on Windows we should receive WM_CLOSE when a Feature is removed already
+
+
+// chunk numbering is off by 2 to account for the registration chunk and binary chunk
+static const int ChunkIDOffset = 2;
+static inline uint32 ToUnrealChunkID( uint32 PackageChunkId )  { return PackageChunkId - ChunkIDOffset; }
+static inline uint32 FromUnrealChunkID( uint32 UnrealChunkID ) { return UnrealChunkID + ChunkIDOffset; }
+
+
+
+FGDKPlatformChunkInstall::FGDKPlatformChunkInstall()
+{
+	UE_LOG(LogChunkInstaller, Display, TEXT("GDK platform chunk installer created"));
+
+	HRESULT hResult;
+
+	// check whether this is an installed package build or not
+	bIsPackaged = XPackageIsPackagedProcess();
+	if (!bIsPackaged)
+	{
+		UE_LOG( LogChunkInstaller, Log, TEXT("not a packaged process - no chunk installation required") );
+		return;
+	}
+
+	// cache package identifier
+	FMemory::Memzero(PackageIdentifier);
+	hResult = XPackageGetCurrentProcessPackageIdentifier(XPACKAGE_IDENTIFIER_MAX_LENGTH, PackageIdentifier);
+	UE_CLOG( FAILED(hResult), LogChunkInstaller, Fatal, TEXT("Failed to get process package identifier. 0x%X"), hResult );
+
+#if !WITH_CHUNKINSTALL_ASYNC_INIT
+#if USE_GDK_PACKAGE_MANIFEST_INIT
+	CreateChunkDataFromPackageManifest();
+#else
+	GDK_SCOPE_NOT_TIME_SENSITIVE(); // (startup only) XPackageEnumeratePackages is not safe to call on a time-sensitive thread
+
+	// create chunk data for the main package
+	auto EnumeratePackageCallback = [](void* Context, const XPackageDetails* Details)
+	{
+		FGDKPlatformChunkInstall* This = (FGDKPlatformChunkInstall*)Context;
+		This->CreateChunkDataForPackage(*Details);
+		return true;
+	};
+
+	hResult = XPackageEnumeratePackages( XPackageKind::Game, XPackageEnumerationScope::ThisOnly, this, EnumeratePackageCallback );
+	UE_CLOG(FAILED(hResult), LogChunkInstaller, Error, TEXT("XPackageEnumeratePackages failed: 0x%X"), hResult);	
+#endif //USE_GDK_PACKAGE_MANIFEST_INIT
+#endif //WITH_CHUNKINSTALL_ASYNC_INIT
+}
+
+
+
+FGDKPlatformChunkInstall::~FGDKPlatformChunkInstall()
+{
+	RemoveFromTicker();
+}
+
+
+bool FGDKPlatformChunkInstall::IsAvailable() const
+{
+	return bIsPackaged;
+}
+
+
+
+
+
+bool FGDKPlatformChunkInstall::MountPaks(uint32 UnrealChunkID)
+{
+	UE_LOG(LogChunkInstaller, Log, TEXT("Mounting pak files from chunk %i."), UnrealChunkID);
+
+	// find and mount all pak files
+	const TArray<FString> PakchunkFiles = IGDKPackageManifestModule::Get().GetPakFilesInChunk(UnrealChunkID);
+	bool bChunkMounted = true;
+	for (const FString& PakchunkFile : PakchunkFiles)
+	{
+		FString PakLocation = TEXT("../../..") / PakchunkFile;
+		PakLocation.ReplaceCharInline(TEXT('\\'), TEXT('/'));
+		IPakFile* PakFile = FCoreDelegates::MountPak.IsBound() ? FCoreDelegates::MountPak.Execute(PakLocation, INDEX_NONE) : nullptr;
+
+		if (PakFile)
+		{
+			UE_LOG(LogChunkInstaller, Log, TEXT("Successfully mounted %s"), *PakchunkFile);
+		}
+		else
+		{
+			UE_LOG(LogChunkInstaller, Log, TEXT("Failed to mount %s"), *PakchunkFile);
+		}
+
+		bChunkMounted &= (PakFile != nullptr);
+	}
+
+	TFunction<void()> BroadcastDelegateFunc = [this, bChunkMounted, UnrealChunkID]()
+	{
+		if (bChunkMounted)
+		{
+			InstallDelegate.Broadcast(UnrealChunkID, true);
+		}
+		else
+		{
+			UE_LOG(LogChunkInstaller, Warning, TEXT("Chunk %i couldn't be mounted."), UnrealChunkID);
+			InstallDelegate.Broadcast(UnrealChunkID, false);
+		}
+	};
+
+	if (IsInGameThread())
+	{
+		BroadcastDelegateFunc();
+	}
+	else if (FTaskGraphInterface::IsRunning())
+	{
+		ExecuteOnGameThread(UE_SOURCE_LOCATION, BroadcastDelegateFunc);
+	}
+
+	return bChunkMounted;
+}
+
+bool FGDKPlatformChunkInstall::UnMountPaks(uint32 UnrealChunkID)
+{
+	UE_LOG(LogChunkInstaller, Log, TEXT("Unmounting pak files from chunk %i."), UnrealChunkID);
+
+	// find and unmount all pak files
+	const TArray<FString> PakchunkFiles = IGDKPackageManifestModule::Get().GetPakFilesInChunk(UnrealChunkID);
+	bool bChunkUnmounted = true;
+	for (const FString& PakchunkFile : PakchunkFiles)
+	{
+		FString PakLocation = TEXT("../../..") / PakchunkFile;
+		PakLocation.ReplaceCharInline(TEXT('\\'), TEXT('/'));
+		bool bPakFileUnmounted = FCoreDelegates::OnUnmountPak.IsBound() ? FCoreDelegates::OnUnmountPak.Execute(PakLocation) : false;
+
+		if (bPakFileUnmounted)
+		{
+			UE_LOG(LogChunkInstaller, Log, TEXT("Successfully unmounted %s"), *PakchunkFile);
+		}
+		else
+		{
+			UE_LOG(LogChunkInstaller, Log, TEXT("Failed to unmount %s"), *PakchunkFile);
+		}
+
+		bChunkUnmounted &= bPakFileUnmounted;
+	}
+
+	if (!bChunkUnmounted)
+	{
+		UE_LOG(LogChunkInstaller, Warning, TEXT("Streaming Install Chunk %i couldn't be unmounted."), UnrealChunkID);
+	}
+
+	return bChunkUnmounted;
+}
+
+void FGDKPlatformChunkInstall::ExternalNotifyChunkAvailable(uint32 PakchunkIndex)
+{
+	uint32 UnrealChunkID = IGDKPackageManifestModule::Get().GetChunkIDFromPakchunkIndex(PakchunkIndex);
+
+	TFunction<void()> BroadcastDelegateFunc = [this, UnrealChunkID]()
+	{
+		InstallDelegate.Broadcast(UnrealChunkID, true);
+	};
+
+	if (IsInGameThread())
+	{
+		BroadcastDelegateFunc();
+	}
+	else if (FTaskGraphInterface::IsRunning())
+	{
+		ExecuteOnGameThread(UE_SOURCE_LOCATION, BroadcastDelegateFunc);
+	}
+}
+
+bool FGDKPlatformChunkInstall::PrioritizePakchunk(int32 PakchunkIndex, EChunkPriority::Type Priority)
+{
+	int32 UnrealChunkID = IGDKPackageManifestModule::Get().GetChunkIDFromPakchunkIndex(PakchunkIndex);
+	return PrioritizeChunk(UnrealChunkID, Priority);
+}
+
+EChunkLocation::Type FGDKPlatformChunkInstall::GetPakchunkLocation(int32 PakchunkIndex)
+{
+	int32 UnrealChunkID = IGDKPackageManifestModule::Get().GetChunkIDFromPakchunkIndex(PakchunkIndex);
+	return GetChunkLocation(UnrealChunkID);
+}
+
+
+
+
+#if WITH_CHUNKINSTALL_ASYNC_INIT
+void FGDKPlatformChunkInstall::AsyncInit( TFunction<void(bool/*bSuccess*/)> OnInitComplete )
+{
+	UE::Tasks::Launch(UE_SOURCE_LOCATION, [this, OnInitComplete = MoveTemp(OnInitComplete)]() mutable
+	{
+#if USE_GDK_PACKAGE_MANIFEST_INIT
+		bool bSuccess = CreateChunkDataFromPackageManifest();
+#else
+		bool bSuccess = true;
+
+		// create chunk data for the main package
+		auto EnumeratePackageCallback = [](void* Context, const XPackageDetails* Details)
+		{
+			FGDKPlatformChunkInstall* This = (FGDKPlatformChunkInstall*)Context;
+			This->CreateChunkDataForPackage(*Details);
+			return true;
+		};
+		HRESULT hResult = XPackageEnumeratePackages( XPackageKind::Game, XPackageEnumerationScope::ThisOnly, this, EnumeratePackageCallback );
+		UE_CLOG(FAILED(hResult), LogChunkInstaller, Error, TEXT("XPackageEnumeratePackages failed: 0x%X"), hResult);
+		bSuccess = SUCCEEDED(hResult);
+#endif // USE_GDK_PACKAGE_MANIFEST_INIT
+
+		// signal completion on the gamethread
+		ExecuteOnGameThread(UE_SOURCE_LOCATION, [this, OnInitComplete = MoveTemp(OnInitComplete), bSuccess]() mutable
+		{
+			bIsInitialized = bSuccess;
+			OnInitComplete(bSuccess);
+		});
+	});
+}
+#endif //WITH_CHUNKINSTALL_ASYNC_INIT
+
+#if !USE_GDK_PACKAGE_MANIFEST_INIT
+void FGDKPlatformChunkInstall::CreateChunkDataForPackage(const XPackageDetails& Details)
+{
+	UE_LOG(LogChunkInstaller, Display, TEXT("Enumerating package %s (%d.%d.%d.%d) %s"), UTF8_TO_TCHAR(Details.packageIdentifier), Details.version.major, Details.version.minor, Details.version.build, Details.version.revision, Details.installing ? TEXT(" - installing") : TEXT("") );
+	FScopeLock Lock(&ChunkDataLock);
+
+	GDK_SCOPE_NOT_TIME_SENSITIVE(); // (startup only) XPackageEnumerateChunkAvailability / XPackageEnumerateFeatures / XPackageGetUserLocale are not safe to call on a time-sensitive thread
+
+	// see if there are any Features
+	bPackageHasFeatures = false;
+	HRESULT hResult = XPackageEnumerateFeatures(PackageIdentifier, &bPackageHasFeatures, [](void* Context, const XPackageFeature* Feature)
+	{
+		// we've found a feature - set bPackageHasFeatures to true and stop enumerating
+		*((bool*)Context) = true; 
+		return false;
+	});
+	UE_CLOG(FAILED(hResult), LogChunkInstaller, Error, TEXT("XPackageEnumerateFeatures failed: 0x%X"), hResult);
+
+	// enumerate all chunks
+	hResult = XPackageEnumerateChunkAvailability( Details.packageIdentifier, XPackageChunkSelectorType::Chunk, this, [](void* Context, const XPackageChunkSelector* Selector, XPackageChunkAvailability Availability)
+	{
+		if ((Selector->chunkId & 0x40000000) == 0) //ignore system chunks
+		{
+			FGDKPlatformChunkInstall* This = (FGDKPlatformChunkInstall*)Context;
+			This->ChunkMonitors.Emplace( ToUnrealChunkID(Selector->chunkId), new FChunkMonitor(*This, *Selector, Availability) );
+		}
+		return true;
+	});
+	UE_CLOG(FAILED(hResult), LogChunkInstaller, Error, TEXT("XPackageEnumerateChunkAvailability (chunks) failed: 0x%X"), hResult);
+
+	// enumerate all Languages
+	hResult = XPackageEnumerateChunkAvailability( Details.packageIdentifier, XPackageChunkSelectorType::Language, this, [](void* Context, const XPackageChunkSelector* Selector, XPackageChunkAvailability Availability)
+	{
+		FGDKPlatformChunkInstall* This = (FGDKPlatformChunkInstall*)Context;
+		
+		FName NamedChunk(Selector->language);
+		This->LanguageChunks.Emplace(NamedChunk);
+		This->NamedChunkMonitors.Emplace(NamedChunk, new FChunkMonitor(*This, *Selector, Availability) );
+		return true;
+	});
+	UE_CLOG(FAILED(hResult), LogChunkInstaller, Error, TEXT("XPackageEnumerateChunkAvailability (languages) failed: 0x%X"), hResult);
+
+
+	if (bPackageHasFeatures)
+	{
+		// enumerate all Features
+		hResult = XPackageEnumerateChunkAvailability( Details.packageIdentifier, XPackageChunkSelectorType::Feature, this, [](void* Context, const XPackageChunkSelector* Selector, XPackageChunkAvailability Availability)
+		{
+			FGDKPlatformChunkInstall* This = (FGDKPlatformChunkInstall*)Context;
+
+			FName NamedChunk(Selector->feature);
+			This->OnDemandChunks.Emplace(NamedChunk);
+			This->NamedChunkMonitors.Emplace(NamedChunk, new FChunkMonitor(*This, *Selector, Availability) );
+			return true;
+		});
+		UE_CLOG(FAILED(hResult), LogChunkInstaller, Error, TEXT("XPackageEnumerateChunkAvailability (features) failed: 0x%X"), hResult);
+	}
+	else
+	{
+		// this package does not use Features - enumerate raw Tags instead
+		hResult = XPackageEnumerateChunkAvailability( Details.packageIdentifier, XPackageChunkSelectorType::Tag, this, [](void* Context, const XPackageChunkSelector* Selector, XPackageChunkAvailability Availability)
+		{
+			FGDKPlatformChunkInstall* This = (FGDKPlatformChunkInstall*)Context;
+
+			FName NamedChunk(Selector->tag);
+			This->OnDemandChunks.Emplace(NamedChunk);
+			This->NamedChunkMonitors.Emplace(NamedChunk, new FChunkMonitor(*This, *Selector, Availability) );
+			return true;
+		});
+		UE_CLOG(FAILED(hResult), LogChunkInstaller, Error, TEXT("XPackageEnumerateChunkAvailability (tags) failed: 0x%X"), hResult);
+	}
+
+	UE_LOG(LogChunkInstaller, Display, TEXT("Finished enumerating package chunks:"));
+	for (const FName& NamedChunk : OnDemandChunks)
+	{
+		UE_LOG(LogChunkInstaller, Display, TEXT("\tOnDemand: %s"), *NamedChunk.ToString());
+	}
+	for (const FName& NamedChunk : LanguageChunks)
+	{
+		UE_LOG(LogChunkInstaller, Display, TEXT("\tLanguage: %s"), *NamedChunk.ToString());
+	}
+
+
+	// read package locale
+	FString PackageLocale;
+	{
+		char PackageLocaleAnsi[LOCALE_NAME_MAX_LENGTH];
+		hResult = XPackageGetUserLocale(LOCALE_NAME_MAX_LENGTH, PackageLocaleAnsi);
+		if (SUCCEEDED(hResult))
+		{
+			PackageLocale = ANSI_TO_TCHAR(PackageLocaleAnsi);
+		}
+		else
+		{
+			UE_CLOG(FAILED(hResult), LogChunkInstaller, Error, TEXT("XPackageGetUserLocale failed: 0x%X"), hResult);
+			PackageLocale = FPlatformMisc::GetDefaultLocale();
+		}
+		UE_LOG(LogChunkInstaller, Display, TEXT("Using package locale: %s"), *PackageLocale);
+	}
+
+	// read stage id overrides from the target settings
+	TMap<FName, FName> StageIdOverrides;
+	{
+		FString RawStageIdOverrides; // StageIdOverrides=(("en","en-US"), ... )
+		GConfig->GetString(FPlatformProperties::GetRuntimeSettingsClassName(), TEXT("StageIdOverrides"), RawStageIdOverrides, GEngineIni);
+		RawStageIdOverrides.ReplaceInline(TEXT("("), TEXT(""));
+		RawStageIdOverrides.ReplaceInline(TEXT(")"), TEXT(""));
+		RawStageIdOverrides.ReplaceInline(TEXT("\""), TEXT(""));
+		RawStageIdOverrides.TrimStartAndEndInline();
+		if (!RawStageIdOverrides.IsEmpty())
+		{
+			TArray<FString> StageIdPairs;
+			RawStageIdOverrides.ParseIntoArray( StageIdPairs, TEXT(","), false );
+			check( (StageIdPairs.Num() % 2) == 0 ); // should have an even number of these as they're keyval pairs
+			for (int32 Index = 0; Index < StageIdPairs.Num(); Index += 2)
+			{
+				const FString& StageId   = StageIdPairs[Index];
+				const FString& CultureId = StageIdPairs[Index+1];
+				if (!CultureId.IsEmpty())
+				{
+					StageIdOverrides.Add(FName(*CultureId), FName(*StageId));
+				}
+			}
+		}
+	}
+
+
+	const TArray<FGDKPackageManifestChunk>& AllChunks = IGDKPackageManifestModule::Get().GetChunks();
+	UE_LOG(LogChunkInstaller, Display, TEXT("%d chunks in manifest"), AllChunks.Num());
+
+	for (const FGDKPackageManifestChunk& Chunk : AllChunks)
+	{
+		// record any Initial chunks, so that the install state can be queried (mostly for bundle support)
+		// if using features, it will be necessary to manually add the Tag to the launch chunks in the ini file as the field is disabled in the editor for launch chunks
+		if (Chunk.bIsInitial)
+		{
+			FName NamedChunk(*Chunk.Tag);
+			InitialNamedChunksMap.Add(NamedChunk, Chunk.UnrealChunkID);
+			UE_LOG(LogChunkInstaller, Display, TEXT("\tInitial chunk %s -> %d"), *NamedChunk.ToString(), Chunk.UnrealChunkID);
+		}
+
+		// cache all named chunks that are not available in the current locale (for use by IsNamedChunkForCurrentLocale)
+		if (!Chunk.Languages.IsEmpty())
+		{
+			TArray<FString> ChunkLanguages;
+			Chunk.Languages.ParseIntoArray(ChunkLanguages, TEXT(";"));
+			for (const FString& Language : ChunkLanguages)
+			{
+				UE_LOG(LogChunkInstaller, Display, TEXT("\tlanguage %s for chunk %d"), *Language, Chunk.UnrealChunkID);
+
+				if (Language != PackageLocale)
+				{
+					FName NamedChunk(Language);
+					SecondaryLanguageNamedChunks.Add(NamedChunk);
+					if (StageIdOverrides.Contains(NamedChunk))
+					{
+						SecondaryLanguageNamedChunks.Add(StageIdOverrides[NamedChunk]); // this allows IsNamedChunkForCurrentLocale to also work with UE Stage Ids in addition to GDK Culture Ids
+					}
+				}
+			}
+		}
+	}
+
+	UE_LOG(LogChunkInstaller, Display, TEXT("Initialize complete. OnDemand: %d, Language: %d, Monitors: %d, NamedMonitors: %d"), OnDemandChunks.Num(), LanguageChunks.Num(), ChunkMonitors.Num(), NamedChunkMonitors.Num());
+}
+#endif //!USE_GDK_PACKAGE_MANIFEST_INIT
+
+#if USE_GDK_PACKAGE_MANIFEST_INIT
+bool FGDKPlatformChunkInstall::CreateChunkDataFromPackageManifest()
+{
+	bool bSuccess = true;
+
+	const TArray<FGDKPackageManifestChunk>& ManifestChunks = IGDKPackageManifestModule::Get().GetChunks();
+	const TArray<FGDKPackageManifestFeature>& ManifestFeatures = IGDKPackageManifestModule::Get().GetFeatures();
+	const TArray<FName>& ManifestLanguages = IGDKPackageManifestModule::Get().GetLanguages();
+	const TArray<FName>& ManifestTags = IGDKPackageManifestModule::Get().GetTags();
+
+	UE_LOG(LogChunkInstaller, Display, TEXT("Enumerating package %s"), UTF8_TO_TCHAR(PackageIdentifier) );
+	FScopeLock Lock(&ChunkDataLock);
+
+	GDK_SCOPE_NOT_TIME_SENSITIVE(); // (startup only) XPackageFindChunkAvailability / XPackageGetUserLocale are not safe to call on a time-sensitive thread
+
+	// see if there are any Features
+	bPackageHasFeatures = ManifestFeatures.Num() > 0;
+
+	// enumerate all Chunks
+	for (const FGDKPackageManifestChunk& ManifestChunk : ManifestChunks)
+	{
+		XPackageChunkSelector ChunkSelector;
+		ChunkSelector.type = XPackageChunkSelectorType::Chunk;
+		ChunkSelector.chunkId = FromUnrealChunkID(ManifestChunk.UnrealChunkID);
+
+		XPackageChunkAvailability ChunkAvailability = XPackageChunkAvailability::Unavailable;
+		HRESULT hResult = XPackageFindChunkAvailability(PackageIdentifier, 1, &ChunkSelector, &ChunkAvailability);
+		if (FAILED(hResult))
+		{
+			bSuccess = false;
+			UE_LOG(LogChunkInstaller, Error, TEXT("XPackageFindChunkAvailability (chunk %d) failed: 0x%X"), ManifestChunk.UnrealChunkID, hResult);
+			continue;
+		}
+
+		ChunkMonitors.Emplace( ManifestChunk.UnrealChunkID, new FChunkMonitor(*this, ChunkSelector, ChunkAvailability) );
+	}
+
+	// enumerate all Languages
+	for (const FName& ManifestLanguage : ManifestLanguages)
+	{
+		FTCHARToUTF8 UTF8Language(*ManifestLanguage.ToString());
+
+		XPackageChunkSelector ChunkSelector;
+		ChunkSelector.type = XPackageChunkSelectorType::Language;
+		ChunkSelector.language = UTF8Language.Get();
+
+		XPackageChunkAvailability ChunkAvailability = XPackageChunkAvailability::Unavailable;
+		HRESULT hResult = XPackageFindChunkAvailability(PackageIdentifier, 1, &ChunkSelector, &ChunkAvailability);
+		if (FAILED(hResult))
+		{
+			bSuccess = false;
+			UE_LOG(LogChunkInstaller, Error, TEXT("XPackageFindChunkAvailability (language %s) failed: 0x%X"), *ManifestLanguage.ToString(), hResult);
+			continue;
+		}
+
+		LanguageChunks.Emplace(ManifestLanguage);
+		NamedChunkMonitors.Emplace(ManifestLanguage, new FChunkMonitor(*this, ChunkSelector, ChunkAvailability) );
+	}
+
+
+	if (bPackageHasFeatures)
+	{
+		// enumerate all Features
+		for (const FGDKPackageManifestFeature& ManifestFeature : ManifestFeatures)
+		{
+			FTCHARToUTF8 UTF8Feature(*ManifestFeature.Id);
+
+			XPackageChunkSelector ChunkSelector;
+			ChunkSelector.type = XPackageChunkSelectorType::Feature;
+			ChunkSelector.feature = UTF8Feature.Get();
+
+			XPackageChunkAvailability ChunkAvailability = XPackageChunkAvailability::Unavailable;
+			HRESULT hResult = XPackageFindChunkAvailability(PackageIdentifier, 1, &ChunkSelector, &ChunkAvailability);
+			if (FAILED(hResult))
+			{
+				bSuccess = false;
+				UE_LOG(LogChunkInstaller, Error, TEXT("XPackageFindChunkAvailability (feature %s) failed: 0x%X"), *ManifestFeature.Id, hResult);
+				continue;
+			}
+
+			FName NamedChunk(ManifestFeature.Id);
+			OnDemandChunks.Emplace(NamedChunk);
+			NamedChunkMonitors.Emplace(NamedChunk, new FChunkMonitor(*this, ChunkSelector, ChunkAvailability) );
+		}
+	}
+	else
+	{
+		// this package does not use Features - enumerate raw Tags instead
+		for (const FName& ManifestTag : ManifestTags)
+		{
+			FTCHARToUTF8 UTF8Tag(*ManifestTag.ToString());
+
+			XPackageChunkSelector ChunkSelector;
+			ChunkSelector.type = XPackageChunkSelectorType::Tag;
+			ChunkSelector.tag = UTF8Tag.Get();
+
+			XPackageChunkAvailability ChunkAvailability = XPackageChunkAvailability::Unavailable;
+			HRESULT hResult = XPackageFindChunkAvailability(PackageIdentifier, 1, &ChunkSelector, &ChunkAvailability);
+			if (FAILED(hResult))
+			{
+				bSuccess = false;
+				UE_LOG(LogChunkInstaller, Error, TEXT("XPackageFindChunkAvailability (tag %s) failed: 0x%X"), *ManifestTag.ToString(), hResult);
+				continue;
+			}
+
+			OnDemandChunks.Emplace(ManifestTag);
+			NamedChunkMonitors.Emplace(ManifestTag, new FChunkMonitor(*this, ChunkSelector, ChunkAvailability) );
+		}
+	}
+
+	UE_LOG(LogChunkInstaller, Display, TEXT("Finished enumerating package chunks:"));
+	for (const FName& NamedChunk : OnDemandChunks)
+	{
+		UE_LOG(LogChunkInstaller, Display, TEXT("\tOnDemand: %s"), *NamedChunk.ToString());
+	}
+	for (const FName& NamedChunk : LanguageChunks)
+	{
+		UE_LOG(LogChunkInstaller, Display, TEXT("\tLanguage: %s"), *NamedChunk.ToString());
+	}
+
+
+	// read package locale
+	FString PackageLocale;
+	{
+		char PackageLocaleAnsi[LOCALE_NAME_MAX_LENGTH];
+		HRESULT hResult = XPackageGetUserLocale(LOCALE_NAME_MAX_LENGTH, PackageLocaleAnsi);
+		if (SUCCEEDED(hResult))
+		{
+			PackageLocale = ANSI_TO_TCHAR(PackageLocaleAnsi);
+		}
+		else
+		{
+			bSuccess = false;
+			UE_CLOG(FAILED(hResult), LogChunkInstaller, Error, TEXT("XPackageGetUserLocale failed: 0x%X"), hResult);
+			PackageLocale = FPlatformMisc::GetDefaultLocale();
+		}
+		UE_LOG(LogChunkInstaller, Display, TEXT("Using package locale: %s"), *PackageLocale);
+	}
+
+	// read stage id overrides from the target settings
+	TMap<FName, FName> StageIdOverrides;
+	{
+		FString RawStageIdOverrides; // StageIdOverrides=(("en","en-US"), ... )
+		GConfig->GetString(FPlatformProperties::GetRuntimeSettingsClassName(), TEXT("StageIdOverrides"), RawStageIdOverrides, GEngineIni);
+		RawStageIdOverrides.ReplaceInline(TEXT("("), TEXT(""));
+		RawStageIdOverrides.ReplaceInline(TEXT(")"), TEXT(""));
+		RawStageIdOverrides.ReplaceInline(TEXT("\""), TEXT(""));
+		RawStageIdOverrides.TrimStartAndEndInline();
+		if (!RawStageIdOverrides.IsEmpty())
+		{
+			TArray<FString> StageIdPairs;
+			RawStageIdOverrides.ParseIntoArray( StageIdPairs, TEXT(","), false );
+			check( (StageIdPairs.Num() % 2) == 0 ); // should have an even number of these as they're keyval pairs
+			for (int32 Index = 0; Index < StageIdPairs.Num(); Index += 2)
+			{
+				const FString& StageId   = StageIdPairs[Index];
+				const FString& CultureId = StageIdPairs[Index+1];
+				if (!CultureId.IsEmpty())
+				{
+					StageIdOverrides.Add(FName(*CultureId), FName(*StageId));
+				}
+			}
+		}
+	}
+
+
+	UE_LOG(LogChunkInstaller, Display, TEXT("%d chunks in manifest"), ManifestChunks.Num());
+
+	for (const FGDKPackageManifestChunk& Chunk : ManifestChunks)
+	{
+		// record any Initial chunks, so that the install state can be queried (mostly for bundle support)
+		// if using features, it will be necessary to manually add the Tag to the launch chunks in the ini file as the field is disabled in the editor for launch chunks
+		if (Chunk.bIsInitial)
+		{
+			FName NamedChunk(*Chunk.Tag);
+			InitialNamedChunksMap.Add(NamedChunk, Chunk.UnrealChunkID);
+			UE_LOG(LogChunkInstaller, Display, TEXT("\tInitial chunk %s -> %d"), *NamedChunk.ToString(), Chunk.UnrealChunkID);
+		}
+
+		// cache all named chunks that are not available in the current locale (for use by IsNamedChunkForCurrentLocale)
+		if (!Chunk.Languages.IsEmpty())
+		{
+			TArray<FString> ChunkLanguages;
+			Chunk.Languages.ParseIntoArray(ChunkLanguages, TEXT(";"));
+			for (const FString& Language : ChunkLanguages)
+			{
+				UE_LOG(LogChunkInstaller, Display, TEXT("\tlanguage %s for chunk %d"), *Language, Chunk.UnrealChunkID);
+
+				if (Language != PackageLocale)
+				{
+					FName NamedChunk(Language);
+					SecondaryLanguageNamedChunks.Add(NamedChunk);
+					if (StageIdOverrides.Contains(NamedChunk))
+					{
+						SecondaryLanguageNamedChunks.Add(StageIdOverrides[NamedChunk]); // this allows IsNamedChunkForCurrentLocale to also work with UE Stage Ids in addition to GDK Culture Ids
+					}
+				}
+			}
+		}
+	}
+
+	UE_LOG(LogChunkInstaller, Display, TEXT("Initialize complete. OnDemand: %d, Language: %d, Monitors: %d, NamedMonitors: %d"), OnDemandChunks.Num(), LanguageChunks.Num(), ChunkMonitors.Num(), NamedChunkMonitors.Num());
+	return bSuccess;
+}
+#endif //USE_GDK_PACKAGE_MANIFEST_INIT
+
+
+
+FGDKPlatformChunkInstall::FChunkMonitor* FGDKPlatformChunkInstall::FindChunkMonitor( uint32 UnrealChunkID )
+{
+#if WITH_CHUNKINSTALL_ASYNC_INIT
+	checkf(bIsInitialized, TEXT("wait for AsyncInit first"));
+#endif
+
+	TUniquePtr<FChunkMonitor>* ResultPtr = ChunkMonitors.Find(UnrealChunkID);
+	return ResultPtr ? ResultPtr->Get() : nullptr;
+}
+
+FGDKPlatformChunkInstall::FChunkMonitor* FGDKPlatformChunkInstall::FindChunkMonitor( const FName NamedChunk )
+{
+#if WITH_CHUNKINSTALL_ASYNC_INIT
+	checkf(bIsInitialized, TEXT("wait for AsyncInit first"));
+#endif
+
+	TUniquePtr<FChunkMonitor>* ResultPtr = NamedChunkMonitors.Find(NamedChunk);
+	return ResultPtr ? ResultPtr->Get() : nullptr;
+}
+
+void FGDKPlatformChunkInstall::OnChunkInstallationChanged( uint32 UnrealChunkID, bool bIsInstalled )
+{
+	UE_LOG(LogChunkInstaller, Log, TEXT("Chunk %d %s"), UnrealChunkID, bIsInstalled ? TEXT("installed") : TEXT("uninstalled") );
+
+	if (bAutoMountPaks)
+	{
+		if (bIsInstalled)
+		{
+			MountPaks(UnrealChunkID);
+		}
+		else
+		{
+			UnMountPaks(UnrealChunkID);
+		}
+	}
+}
+
+void FGDKPlatformChunkInstall::OnNamedChunkInstallationChanged( const FName NamedChunk, bool bIsInstalled )
+{
+	if (!bIsInstalled)
+	{
+		//nb. we are inside a FScopeLock Lock(&ChunkDataLock) from the caller, OnProgressChanged
+
+		if (CurrentUninstallingChunks.Contains(NamedChunk))
+		{
+			CurrentUninstallingChunks.Remove(NamedChunk);
+		}
+		else if (bPackageHasFeatures && GGDKChunkInstallTerminateGameOnExternalFeatureRemoval)
+		{
+			UE_LOG(LogChunkInstaller, Log, TEXT("Named chunk %s removed externally - terminating the game"), *NamedChunk.ToString());	
+			RequestEngineExit(FString::Printf(TEXT("Named chunk %s removed externally"), *NamedChunk.ToString()));
+		}
+	}
+
+
+	UE_LOG(LogChunkInstaller, Log, TEXT("Named chunk %s %s"), *NamedChunk.ToString(), bIsInstalled ? TEXT("installed") : TEXT("uninstalled") );
+
+	EChunkLocation::Type Location = (bIsInstalled ? EChunkLocation::LocalFast : EChunkLocation::NotAvailable);
+	if (IsInGameThread())
+	{
+		DoNamedChunkCompleteCallbacks(NamedChunk, Location, true);
+	}
+	else if (FTaskGraphInterface::IsRunning())
+	{
+		ExecuteOnGameThread(UE_SOURCE_LOCATION, [this, NamedChunk, Location]
+		{
+			DoNamedChunkCompleteCallbacks(NamedChunk, Location, true);
+		});
+	}
+}
+
+
+
+
+
+EChunkLocation::Type FGDKPlatformChunkInstall::GetChunkLocation(uint32 UnrealChunkID)
+{
+	const FChunkMonitor* ChunkMonitor = FindChunkMonitor(UnrealChunkID);
+	if (ChunkMonitor != nullptr)
+	{
+		return ChunkMonitor->GetLocation();
+	}
+
+	return EChunkLocation::DoesNotExist;
+}
+
+EChunkLocation::Type FGDKPlatformChunkInstall::GetNamedChunkLocation(const FName NamedChunk)
+{
+	if (InitialNamedChunksMap.Contains(NamedChunk))
+	{
+		return EChunkLocation::LocalFast;
+	}
+
+	const FChunkMonitor* ChunkMonitor = FindChunkMonitor(NamedChunk);
+	if (ChunkMonitor != nullptr)
+	{
+		return ChunkMonitor->GetLocation();
+	}
+
+	return EChunkLocation::DoesNotExist;
+}
+
+
+float FGDKPlatformChunkInstall::GetChunkProgress( uint32 PakchunkIndex, EChunkProgressReportingType::Type ReportType )
+{
+	uint32 UnrealChunkID = IGDKPackageManifestModule::Get().GetChunkIDFromPakchunkIndex(PakchunkIndex);
+	const FChunkMonitor* ChunkMonitor = FindChunkMonitor(UnrealChunkID);
+	if (ChunkMonitor != nullptr)
+	{
+		return ChunkMonitor->GetProgress(ReportType);
+	}
+
+	return 0.0f;
+}
+
+float FGDKPlatformChunkInstall::GetNamedChunkProgress(const FName NamedChunk, EChunkProgressReportingType::Type ReportType)
+{
+	if (InitialNamedChunksMap.Contains(NamedChunk))
+	{
+		return (ReportType == EChunkProgressReportingType::PercentageComplete) ? 100.0f : 0.0f;
+	}
+
+	const FChunkMonitor* ChunkMonitor = FindChunkMonitor(NamedChunk);
+	if (ChunkMonitor != nullptr)
+	{
+		return ChunkMonitor->GetProgress(ReportType);
+	}
+
+	return 0.0f;
+}
+
+
+
+
+bool FGDKPlatformChunkInstall::PrioritizeChunk( uint32 UnrealChunkID, EChunkPriority::Type Priority )
+{
+	FChunkMonitor* ChunkMonitor = FindChunkMonitor(UnrealChunkID);
+	if (ChunkMonitor != nullptr)
+	{
+		return ChunkMonitor->SetPriority(Priority);
+	}
+
+	return false;
+}
+
+bool FGDKPlatformChunkInstall::PrioritizeNamedChunk(const FName NamedChunk, EChunkPriority::Type Priority)
+{
+	FChunkMonitor* ChunkMonitor = FindChunkMonitor(NamedChunk);
+	if (ChunkMonitor != nullptr)
+	{
+		return ChunkMonitor->SetPriority(Priority);
+	}
+
+	return false;
+}
+
+
+bool FGDKPlatformChunkInstall::IsNamedChunkInProgress(const FName NamedChunk)
+{
+	FChunkMonitor* ChunkMonitor = FindChunkMonitor(NamedChunk);
+	if (ChunkMonitor == nullptr)
+	{
+		return false;
+	}
+
+	GDK_SCOPE_NOT_TIME_SENSITIVE(); // XPackageFindChunkAvailability is not safe to call on a time-sensitive thread
+
+	// check the availability of the given chunk
+	XPackageChunkAvailability InstallationState;
+	HRESULT hResult = XPackageFindChunkAvailability( PackageIdentifier, 1, &ChunkMonitor->Selector, &InstallationState );
+	UE_CLOG(FAILED(hResult), LogChunkInstaller, Error, TEXT("Failed to query named chunk availability for %s - 0x%X"), *NamedChunk.ToString(), hResult);
+
+	return SUCCEEDED(hResult) && InstallationState == XPackageChunkAvailability::Pending;
+}
+
+
+bool FGDKPlatformChunkInstall::Tick( float DeltaSeconds )
+{
+	if (PendingInstallChunks.Num() > 0)
+	{
+		// async callback
+		auto AsyncInstallComplete = []( XAsyncBlock* AsyncBlock)
+		{
+			// collect result
+			XPackageInstallationMonitorHandle Monitor;
+			HRESULT hResult = XPackageInstallChunksResult(AsyncBlock, &Monitor);
+			delete AsyncBlock;
+			if (FAILED(hResult))
+			{
+				UE_LOG(LogChunkInstaller, Error, TEXT("Async named chunk install failed to start - 0x%X"), hResult);
+			}
+			else
+			{
+				// FChunkMonitor takes care of the actual tracking so we can just close the handle here
+				XPackageCloseInstallationMonitorHandle(Monitor); 
+			}
+		};
+
+		// only request chunks that can be installed... otherwise signal the callback immediately
+		TSet<FName> ChunksToInstall;
+		for (FName NamedChunk : PendingInstallChunks)
+		{
+			const FChunkMonitor* ChunkMonitor = FindChunkMonitor(NamedChunk);
+			if (ensure(ChunkMonitor) && ChunkMonitor->bInstalled)
+			{		
+				UE_LOG(LogChunkInstaller, Log, TEXT("Named chunk %s is already installed"), *NamedChunk.ToString() );
+				DoNamedChunkCompleteCallbacks(NamedChunk, EChunkLocation::LocalFast, true);
+			}
+			else if (UnavailableChunks.Contains(NamedChunk))
+			{
+				UE_LOG(LogChunkInstaller, Log, TEXT("Named chunk %s is not available in the current locale/device"), *NamedChunk.ToString() );
+				DoNamedChunkCompleteCallbacks(NamedChunk, EChunkLocation::DoesNotExist, false);
+			}
+			else
+			{
+				ChunksToInstall.Add(NamedChunk);
+			}
+		}
+
+
+		if (ChunksToInstall.Num() > 0)
+		{
+			// prepare the async chunk install
+			XAsyncBlock* AsyncBlock = new XAsyncBlock{};
+			AsyncBlock->queue = FGDKAsyncTaskQueue::GetGenericQueue();
+			AsyncBlock->callback = AsyncInstallComplete;
+
+			// install the chunks
+			FGDKCustomChunkSelector Chunks(ChunksToInstall, this);
+			HRESULT hResult = XPackageInstallChunksAsync( PackageIdentifier, Chunks.Num(), Chunks.Get(), GDK_CHUNKINSTALL_UPDATEINTERVAL_MS, (GGDKChunkInstallSuppressUserConfirmation != 0), AsyncBlock );
+			if (FAILED(hResult))
+			{
+				delete AsyncBlock;
+				UE_LOG(LogChunkInstaller, Error, TEXT("Failed to start async installing named chunks %s - 0x%X"), *Chunks.ToString(), hResult);
+				for (FName NamedChunk : ChunksToInstall)
+				{
+					DoNamedChunkCompleteCallbacks(NamedChunk, EChunkLocation::DoesNotExist, false);
+				}
+
+			}
+		}
+
+		PendingInstallChunks.Reset();
+	}
+
+	if (PendingUninstallChunks.Num() > 0)
+	{
+		// only request chunks that can be uninstalled... otherwise signal the callback immediately
+		TSet<FName> ChunksToUninstall;
+		for (FName NamedChunk : PendingUninstallChunks)
+		{
+			const FChunkMonitor* ChunkMonitor = FindChunkMonitor(NamedChunk);
+			if (ensure(ChunkMonitor) && !ChunkMonitor->bInstalled)
+			{		
+				UE_LOG(LogChunkInstaller, Log, TEXT("Named chunk %s is already uninstalled"), *NamedChunk.ToString() );
+				DoNamedChunkCompleteCallbacks(NamedChunk, EChunkLocation::NotAvailable, true);
+			}
+			else if (UnavailableChunks.Contains(NamedChunk))
+			{
+				UE_LOG(LogChunkInstaller, Log, TEXT("Named chunk %s is not available in the current locale/device"), *NamedChunk.ToString() );
+				DoNamedChunkCompleteCallbacks(NamedChunk, EChunkLocation::DoesNotExist, false);
+			}
+			else
+			{
+				ChunksToUninstall.Add(NamedChunk);
+			}
+		}
+
+
+		if (ChunksToUninstall.Num() > 0)
+		{
+			GDK_SCOPE_NOT_TIME_SENSITIVE(); // XPackageUninstallChunks is not safe to call on a time-sensitive thread
+
+			// uninstall the chunks
+			FGDKCustomChunkSelector Chunks(ChunksToUninstall, this);
+			HRESULT hResult = XPackageUninstallChunks( PackageIdentifier, Chunks.Num(), Chunks.Get() );
+			if (FAILED(hResult))
+			{
+				UE_LOG(LogChunkInstaller, Error, TEXT("Failed to start uninstalling named chunks %s - 0x%X"), *Chunks.ToString(), hResult);
+				for (FName NamedChunk : ChunksToUninstall)
+				{
+					DoNamedChunkCompleteCallbacks(NamedChunk, EChunkLocation::LocalFast, false);
+				}
+			}
+			else
+			{
+				FScopeLock Lock(&ChunkDataLock);
+				CurrentUninstallingChunks = CurrentUninstallingChunks.Union(ChunksToUninstall);
+			}
+		}
+		PendingUninstallChunks.Reset();
+	}
+
+	// stop ticking
+	RemoveFromTicker();
+	return false;
+}
+
+
+bool FGDKPlatformChunkInstall::InstallNamedChunks(const TArrayView<const FName>& NamedChunks)
+{
+#if WITH_CHUNKINSTALL_ASYNC_INIT
+	checkf(bIsInitialized, TEXT("wait for AsyncInit first"));
+#endif
+
+	// early out
+	if (!bIsPackaged || NamedChunks.Num() == 0)
+	{
+		DoNamedChunkCompleteCallbacks(NamedChunks, EChunkLocation::DoesNotExist, false);
+		return false;
+	}
+
+	// validate requested chunks
+	TSet<FName> RequestedChunks;
+	for (FName NamedChunk : NamedChunks)
+	{
+		if (InitialNamedChunksMap.Contains(NamedChunk))
+		{
+			UE_LOG(LogChunkInstaller, Error, TEXT("launch chunk always installed: %s"), *NamedChunk.ToString());
+			DoNamedChunkCompleteCallbacks(NamedChunk, EChunkLocation::LocalFast, true);
+		}
+		else if (NamedChunkMonitors.Contains(NamedChunk))
+		{
+			RequestedChunks.Add(NamedChunk);
+		}
+		else
+		{
+			UE_LOG(LogChunkInstaller, Error, TEXT("unknown named chunk: %s"), *NamedChunk.ToString());
+			DoNamedChunkCompleteCallbacks(NamedChunk, EChunkLocation::DoesNotExist, false);
+		}
+	}
+
+	// update pending chunk sets
+	PendingUninstallChunks = PendingUninstallChunks.Difference(RequestedChunks); // remove chunks from uninstall request
+	PendingInstallChunks = PendingInstallChunks.Union(RequestedChunks); // add chunks to install request
+
+	// request a tick if there's anything to do
+	if (PendingUninstallChunks.Num() > 0 || PendingInstallChunks.Num() > 0)
+	{
+		AddToTicker();
+		return true;
+	}
+
+	return false;
+
+}
+
+
+
+bool FGDKPlatformChunkInstall::UninstallNamedChunks(const TArrayView<const FName>& NamedChunks)
+{
+#if WITH_CHUNKINSTALL_ASYNC_INIT
+	checkf(bIsInitialized, TEXT("wait for AsyncInit first"));
+#endif
+
+	// early out
+	if (!bIsPackaged || NamedChunks.Num() == 0)
+	{
+		DoNamedChunkCompleteCallbacks(NamedChunks, EChunkLocation::DoesNotExist, false);
+		return false;
+	}
+
+	// validate requested chunks
+	TSet<FName> RequestedChunks;
+	for (FName NamedChunk : NamedChunks)
+	{
+		if (InitialNamedChunksMap.Contains(NamedChunk))
+		{
+			UE_LOG(LogChunkInstaller, Error, TEXT("cannot remove launch chunk: %s"), *NamedChunk.ToString());
+			DoNamedChunkCompleteCallbacks(NamedChunk, EChunkLocation::LocalFast, false);
+		}
+		else if (NamedChunkMonitors.Contains(NamedChunk))
+		{
+			RequestedChunks.Add(NamedChunk);
+		}
+		else
+		{
+			UE_LOG(LogChunkInstaller, Error, TEXT("unknown named chunk: %s"), *NamedChunk.ToString());
+			DoNamedChunkCompleteCallbacks(NamedChunk, EChunkLocation::DoesNotExist, false);
+		}
+	}
+
+	// update pending chunk sets
+	PendingUninstallChunks = PendingUninstallChunks.Union(RequestedChunks); // add chunks to uninstall request
+	PendingInstallChunks = PendingInstallChunks.Difference(RequestedChunks); // remove chunks from install request
+
+	// request a tick if there's anything to do
+	if (PendingUninstallChunks.Num() > 0 || PendingInstallChunks.Num() > 0)
+	{
+		AddToTicker();
+		return true;
+	}
+
+	return false;
+
+}
+
+
+
+TArray<FName> FGDKPlatformChunkInstall::GetNamedChunksByType(ENamedChunkType NamedChunkType) const
+{
+#if WITH_CHUNKINSTALL_ASYNC_INIT
+	checkf(bIsInitialized, TEXT("wait for AsyncInit first"));
+#endif
+
+	switch(NamedChunkType)
+	{
+		case ENamedChunkType::OnDemand:	return OnDemandChunks;
+		case ENamedChunkType::Language:	return LanguageChunks;
+	}
+	checkNoEntry();
+	return TArray<FName>();
+}
+
+ENamedChunkType FGDKPlatformChunkInstall::GetNamedChunkType(const FName NamedChunk) const
+{
+#if WITH_CHUNKINSTALL_ASYNC_INIT
+	checkf(bIsInitialized, TEXT("wait for AsyncInit first"));
+#endif
+
+	if (OnDemandChunks.Contains(NamedChunk))
+	{
+		return ENamedChunkType::OnDemand;
+	}
+	else if (LanguageChunks.Contains(NamedChunk))
+	{
+		return ENamedChunkType::Language;
+	}
+	else
+	{
+		return ENamedChunkType::Invalid;
+	}
+}
+
+
+
+// NOTE: this will return all *potential* files for the given named chunk (e.g. tag # (AND) operator is not considered for Features & Recipes)
+bool FGDKPlatformChunkInstall::GetPakFilesInNamedChunk( const FName NamedChunk, TArray<FString>& OutFilesInChunk) const
+{
+#if WITH_CHUNKINSTALL_ASYNC_INIT
+	checkf(bIsInitialized, TEXT("wait for AsyncInit first"));
+#endif
+
+	OutFilesInChunk.Reset();
+
+	if (bPackageHasFeatures)
+	{
+		// get all chunks that comprise this feature
+		TArray<int32> UnrealChunkIDs;
+		if (!IGDKPackageManifestModule::Get().GetUnrealChunkIDsByFeature(NamedChunk.ToString(), UnrealChunkIDs))
+		{
+			return false;
+		}
+
+		// collect all the files from these chunks
+		for (int32 UnrealChunkID : UnrealChunkIDs)
+		{
+			OutFilesInChunk.Append( IGDKPackageManifestModule::Get().GetPakFilesInChunk(UnrealChunkID) );
+		}
+
+		return true;
+	}
+	else
+	{
+		// find the chunk with the given tag & collect the files
+		const FString Tag = NamedChunk.ToString();
+		const TArray<FGDKPackageManifestChunk>& AllChunks = IGDKPackageManifestModule::Get().GetChunks();
+		for (const FGDKPackageManifestChunk& Chunk : AllChunks)
+		{
+			if (Chunk.Tag == Tag)
+			{
+				OutFilesInChunk.Append( IGDKPackageManifestModule::Get().GetPakFilesInChunk(Chunk.UnrealChunkID) );
+				return true;
+			}
+		}
+
+		return false;
+	}
+
+}
+
+bool FGDKPlatformChunkInstall::GetNamedChunkInstallationStatus( const FName NamedChunk, FChunkInstallationStatusDetail& OutChunkStatusDetail ) const
+{
+	const FChunkMonitor* ChunkMonitor = FindChunkMonitor(NamedChunk);
+	if (ChunkMonitor != nullptr)
+	{
+		XPackageInstallationProgress Progress;
+		XPackageGetInstallationProgress( ChunkMonitor->InstallationMonitorHandle, &Progress );
+		OutChunkStatusDetail.CurrentInstallSize = Progress.installedBytes;
+		OutChunkStatusDetail.FullInstallSize = Progress.totalBytes;
+		OutChunkStatusDetail.bIsInstalled = ChunkMonitor->bInstalled;
+
+		return true;
+	}
+
+	// check initial chunks (these are always installed)
+	TArray<int32> InitialChunkIDs;
+	InitialNamedChunksMap.MultiFind(NamedChunk, InitialChunkIDs);
+	if (InitialChunkIDs.Num() > 0 )
+	{
+		uint64 TotalBytes = 0;
+		for (int32 InitialChunkID : InitialChunkIDs)
+		{
+			ChunkMonitor = FindChunkMonitor(InitialChunkID);
+			if (ChunkMonitor != nullptr)
+			{
+				XPackageInstallationProgress Progress;
+				XPackageGetInstallationProgress( ChunkMonitor->InstallationMonitorHandle, &Progress );
+				TotalBytes += Progress.totalBytes;
+			}
+		}
+
+		OutChunkStatusDetail.FullInstallSize = TotalBytes;
+		OutChunkStatusDetail.CurrentInstallSize = TotalBytes;
+		OutChunkStatusDetail.bIsInstalled = true;
+		return true;
+	}
+
+	return false;
+}
+
+
+bool FGDKPlatformChunkInstall::IsNamedChunkForCurrentLocale(const FName NamedChunk) const
+{
+#if WITH_CHUNKINSTALL_ASYNC_INIT
+	checkf(bIsInitialized, TEXT("wait for AsyncInit first"));
+#endif
+
+	return !SecondaryLanguageNamedChunks.Contains(NamedChunk);
+}
+
+
+bool FGDKPlatformChunkInstall::SetAutoPakMountingEnabled( bool bEnabled )
+{
+	bAutoMountPaks = bEnabled;
+	return true;
+}
+
+
+
+
+
+
+
+
+
+void FGDKPlatformChunkInstall::AddToTicker()
+{
+	if (!TickHandle.IsValid())
+	{
+		FTSTicker& Ticker = FTSTicker::GetCoreTicker();
+
+		// Register delegate for ticker callback
+		FTickerDelegate TickDelegate = FTickerDelegate::CreateRaw(this, &FGDKPlatformChunkInstall::Tick);
+		TickHandle = Ticker.AddTicker(TickDelegate, 0.1f);
+	}
+}
+
+void FGDKPlatformChunkInstall::RemoveFromTicker()
+{
+	FTSTicker& Ticker = FTSTicker::GetCoreTicker();
+
+	// Unregister ticker delegate
+	if (TickHandle.IsValid())
+	{
+		Ticker.RemoveTicker(TickHandle);
+		TickHandle.Reset();
+	}
+}
+
+
+
+
+
+void FGDKPlatformChunkInstall::DebugDumpChunkState()
+{
+	GDK_SCOPE_NOT_TIME_SENSITIVE(); // (debug only) XPackageFindChunkAvailability is not safe to call on a time-sensitive thread
+
+	FGDKPlatformChunkInstall* ChunkInstall = (FGDKPlatformChunkInstall*)FPlatformMisc::GetPlatformChunkInstall(); // note: not really safe!
+	for (auto& Itr : ChunkInstall->ChunkMonitors)
+	{
+		XPackageChunkAvailability Availability = XPackageChunkAvailability::Unavailable;
+		XPackageFindChunkAvailability( ChunkInstall->PackageIdentifier, 1, &Itr.Value->Selector, &Availability );
+		UE_LOG( LogChunkInstaller, Log, TEXT("%s %s"), *Itr.Value->GetDebugString(), *LexToString(Availability) );
+	}
+	for (auto& Itr : ChunkInstall->NamedChunkMonitors)
+	{
+		XPackageChunkAvailability Availability = XPackageChunkAvailability::Unavailable;
+		XPackageFindChunkAvailability( ChunkInstall->PackageIdentifier, 1, &Itr.Value->Selector, &Availability );
+		UE_LOG( LogChunkInstaller, Log, TEXT("%s %s"), *Itr.Value->GetDebugString(), *LexToString(Availability) );
+	}
+}
+
+
+
+
+
+
+
+
+
+
+
+
+FGDKPlatformChunkInstall::FGDKCustomChunkSelector::FGDKCustomChunkSelector( const TSet<FName>& NamedChunks, FGDKPlatformChunkInstall* Owner )
+{
+	ChunkSelectors.Reserve(NamedChunks.Num());
+
+	for (const FName& NamedChunk : NamedChunks)
+	{
+		FChunkMonitor* ChunkMonitor = Owner->FindChunkMonitor(NamedChunk);
+		if (ChunkMonitor != nullptr)
+		{
+			ChunkSelectors.Emplace( ChunkMonitor->Selector ); //note: copying the selector including a shallow copy of the tag/feature/language string. It is safe because this class is only ever created on the stack & will never outlive the chunk monitors
+			StringTagList += ChunkMonitor->GetName();
+		}
+		else
+		{
+			UE_LOG(LogChunkInstaller, Error, TEXT("Unknown Named Chunk %s"), *NamedChunk.ToString());
+			StringTagList += NamedChunk.ToString() + TEXT("??");
+		}
+		StringTagList += TEXT(", ");
+
+	}
+
+	// clean up debug string
+	StringTagList.RemoveFromEnd( TEXT(", ") );
+}
+
+
+
+
+
+
+
+
+
+
+FGDKPlatformChunkInstall::FChunkMonitor::FChunkMonitor( FGDKPlatformChunkInstall& InOwner, const XPackageChunkSelector& InSelector, XPackageChunkAvailability InInitialAvailability )
+	: bInstalled(false)
+	, bWasInProgress(false)
+	, CachedProgress(0)
+	, Selector(InSelector)
+	, InstallationMonitorHandle(nullptr)
+	, Owner(&InOwner)
+{
+	// need to take a copy of the string
+	if (Selector.type != XPackageChunkSelectorType::Chunk)
+	{
+		// sanity check: Selector data is a union so these should all be the same otherwise the code below 
+		check(Selector.feature == Selector.language);
+		check(Selector.feature == Selector.tag);
+
+		SIZE_T TagLength = FCStringAnsi::Strlen(Selector.tag);
+		SelectorTag.Reset( new char[TagLength+1] );
+		FCStringAnsi::Strncpy( SelectorTag.Get(), Selector.tag, TagLength + 1 );
+		Selector.tag = SelectorTag.Get();
+	}
+
+	HRESULT hResult;
+
+	// if the chunk is already Ready when this is created (on startup) it means the chunk is downloaded already
+	if (InInitialAvailability == XPackageChunkAvailability::Ready)
+	{
+		CachedProgress = 100;
+		bInstalled = true;
+	}
+
+	// create installation monitor for this chunk
+	hResult = XPackageCreateInstallationMonitor( Owner->PackageIdentifier, 1, &Selector, GDK_CHUNKINSTALL_UPDATEINTERVAL_MS, FGDKAsyncTaskQueue::GetGenericQueue(), &InstallationMonitorHandle );
+	UE_CLOG( FAILED(hResult), LogChunkInstaller, Fatal, TEXT("Failed to create installation monitor for %s. 0x%X"), *GetName(), hResult );
+
+	// register callback when package installation state changes
+	auto InstallationProgressChanged = [](void* Context, XPackageInstallationMonitorHandle MonitorHandle)
+	{
+		FChunkMonitor* This = (FChunkMonitor*)Context;
+		This->OnProgressChange();
+	};
+	hResult = XPackageRegisterInstallationProgressChanged( InstallationMonitorHandle, this, InstallationProgressChanged, &InstallationProgressChangedToken );
+	UE_CLOG( FAILED(hResult), LogChunkInstaller, Fatal, TEXT("Failed to register for installation callback for %s. 0x%X"), *GetName(), hResult );
+}
+
+
+FGDKPlatformChunkInstall::FChunkMonitor::~FChunkMonitor()
+{
+	XPackageUnregisterInstallationProgressChanged(InstallationMonitorHandle, InstallationProgressChangedToken, true);
+	XPackageCloseInstallationMonitorHandle(InstallationMonitorHandle);
+}
+
+FString FGDKPlatformChunkInstall::FChunkMonitor::GetDebugString() const
+{
+	return FString::Printf(TEXT("%-32s %4d%%   %-16s "), *GetName(), CachedProgress, bInstalled ? TEXT("installed") : TEXT(""));
+}
+
+FString FGDKPlatformChunkInstall::FChunkMonitor::GetName() const
+{
+	switch (Selector.type)
+	{
+		case XPackageChunkSelectorType::Language: return FString::Printf(TEXT("Language %s"), UTF8_TO_TCHAR(Selector.language));
+		case XPackageChunkSelectorType::Tag:      return FString::Printf(TEXT("Tag %s"), UTF8_TO_TCHAR(Selector.tag));
+		case XPackageChunkSelectorType::Chunk:    return FString::Printf(TEXT("Chunk %d"), Selector.chunkId);
+		case XPackageChunkSelectorType::Feature:  return FString::Printf(TEXT("Feature %s"), UTF8_TO_TCHAR(Selector.feature));
+	}
+	checkNoEntry();
+	return FString();
+}
+
+void FGDKPlatformChunkInstall::FChunkMonitor::OnProgressChange()
+{
+	FScopeLock Lock(&Owner->ChunkDataLock);
+
+	// cache updated chunk installation progress
+	XPackageInstallationProgress Progress;
+	XPackageGetInstallationProgress( InstallationMonitorHandle, &Progress );
+	CachedProgress = (Progress.totalBytes > 0) ? (Progress.installedBytes * 100) / Progress.totalBytes : 0;
+	UE_LOG( LogChunkInstaller, Verbose, TEXT("%s installation progress changed: %llu / %llu bytes (%d%%).%s"), *GetName(), Progress.installedBytes, Progress.totalBytes, CachedProgress, Progress.completed?TEXT(" completed"):TEXT(""));	
+
+	// check installation state
+	bool bComplete = false;
+	if (Progress.completed && !bInstalled && Progress.totalBytes > 0)
+	{
+		bInstalled = true;
+		bComplete = true;
+	}
+	else if ((bInstalled || bWasInProgress) && (Progress.installedBytes == 0))
+	{
+		bInstalled = false;
+		bComplete = true;
+	}
+	else if (Progress.completed && !bInstalled && Progress.totalBytes == 0 && Selector.type != XPackageChunkSelectorType::Chunk)
+	{
+		FName NamedChunk(Selector.tag);
+		Owner->UnavailableChunks.Add(NamedChunk); // most likely the chunk's Device or Language doesn't match the system
+	}
+
+	// signal completion
+	if (bComplete)
+	{
+		bWasInProgress = false;
+
+		if (Selector.type == XPackageChunkSelectorType::Chunk)
+		{
+			const uint32 UnrealChunkID = ToUnrealChunkID(Selector.chunkId);
+			Owner->OnChunkInstallationChanged(UnrealChunkID, bInstalled);
+		}
+		else
+		{
+			FName NamedChunk(Selector.tag);
+			Owner->OnNamedChunkInstallationChanged( NamedChunk, bInstalled );
+		}
+	}
+	else
+	{
+		bWasInProgress |= (Progress.installedBytes > 0);
+	}
+
+}
+
+bool FGDKPlatformChunkInstall::FChunkMonitor::SetPriority(EChunkPriority::Type Priority)
+{
+	// change the priority of this chunk if we need it right away
+	if (Priority == EChunkPriority::Immediate)
+	{
+		GDK_SCOPE_NOT_TIME_SENSITIVE(); // XPackageChangeChunkInstallOrder is not safe to call on a time-sensitive thread
+
+		HRESULT hResult = XPackageChangeChunkInstallOrder( Owner->PackageIdentifier, 1, &Selector );
+		UE_CLOG(FAILED(hResult), LogChunkInstaller, Error, TEXT("Failed to change chunk install order for %s - 0x%X"), *GetName(), hResult );
+		return SUCCEEDED(hResult);
+	}
+
+	return true;
+}
+
+float FGDKPlatformChunkInstall::FChunkMonitor::GetProgress(EChunkProgressReportingType::Type ReportType) const
+{
+	// make sure we can support this
+	if (ReportType != EChunkProgressReportingType::PercentageComplete)
+	{
+		UE_LOG(LogChunkInstaller, Error, TEXT("ChunkProgressReportType not supported: %i"), (int)ReportType);
+		return 0.0f;
+	}
+
+	// return the last known progress
+	return (float)CachedProgress;
+}
+
+EChunkLocation::Type FGDKPlatformChunkInstall::FChunkMonitor::GetLocation() const
+{
+	// return the last known location
+	return bInstalled ? EChunkLocation::LocalFast : EChunkLocation::NotAvailable;
+}
+
+
+
+
+FString LexToString( const XPackageChunkSelector& ChunkSelector )
+{
+	switch(ChunkSelector.type)
+	{
+		case XPackageChunkSelectorType::Chunk:		return FString::Printf( TEXT("Id:%d"), ChunkSelector.chunkId );
+		case XPackageChunkSelectorType::Language:	return FString::Printf( TEXT("Lang:%s"), ANSI_TO_TCHAR(ChunkSelector.language) );
+		case XPackageChunkSelectorType::Tag:		return FString::Printf( TEXT("Tag:%s"), ANSI_TO_TCHAR(ChunkSelector.tag) );
+		case XPackageChunkSelectorType::Feature:	return FString::Printf( TEXT("Feat:%s"), ANSI_TO_TCHAR(ChunkSelector.feature) );
+		default:									return FString::Printf( TEXT("Unknown type:%d"), ChunkSelector.type );
+	}
+}
+
+FString LexToString(XPackageChunkAvailability Availability)
+{
+	switch(Availability)
+	{
+		case XPackageChunkAvailability::Ready:			return TEXT("Ready");
+		case XPackageChunkAvailability::Pending:		return TEXT("Pending");
+		case XPackageChunkAvailability::Installable:	return TEXT("Installable");
+		case XPackageChunkAvailability::Unavailable:	return TEXT("Unvailable");
+		default:										return TEXT("Invalid");
+	}
+}
+
+#endif //WITH_GRDK
diff --git a/Engine/Plugins/Runtime/Microsoft/GDKPackageChunkInstall/Source/GDKPackageChunkInstall/Public/GDKPlatformChunkInstall.h b/Engine/Plugins/Runtime/Microsoft/GDKPackageChunkInstall/Source/GDKPackageChunkInstall/Public/GDKPlatformChunkInstall.h
new file mode 100644
--- /dev/null
+++ b/Engine/Plugins/Runtime/Microsoft/GDKPackageChunkInstall/Source/GDKPackageChunkInstall/Public/GDKPlatformChunkInstall.h
@@ -0,0 +1,179 @@
+// Copyright Epic Games, Inc. All Rights Reserved.
+
+#pragma once
+
+#if WITH_GRDK
+#include "GenericPlatform/GenericPlatformChunkInstall.h"
+#include "Containers/Map.h"
+#include "Containers/Set.h"
+#include "Containers/Ticker.h"
+#include "Templates/UniquePtr.h"
+
+#include "Microsoft/AllowMicrosoftPlatformTypes.h"
+THIRD_PARTY_INCLUDES_START
+#include <XPackage.h>
+#include <atomic>
+THIRD_PARTY_INCLUDES_END
+#include "Microsoft/HideMicrosoftPlatformTypes.h"
+
+#define UE_API GDKPACKAGECHUNKINSTALL_API
+
+// Set to 1 in your target.cs file to enable asyncronous initialization for faster engine startup
+// You must call AsyncInit and wait for the callback before using any chunk install API functions
+#ifndef WITH_CHUNKINSTALL_ASYNC_INIT
+	#define WITH_CHUNKINSTALL_ASYNC_INIT 0
+#endif
+
+// Set to 1 in your target.cs file to use the package manifest when initializing, rather than enumerating the XPackage.
+// This seems to be faster and may work around very rare issues where the XPackage enumeration callback fails.
+// This may become the default once it has been proven at scale.
+// (Note that when this is enabled, Named Chunks will not be available for package Languages that are not referenced in IntelligentDeliveryChunks' StageIds. Previously these would have been treated as if they were installed.)
+#ifndef USE_GDK_PACKAGE_MANIFEST_INIT
+	#define USE_GDK_PACKAGE_MANIFEST_INIT 0
+#endif
+
+/**
+ * GDK implementation of FGenericPlatformChunkInstall with support for Recipes & Features
+ */
+class FGDKPlatformChunkInstall : public FGenericPlatformChunkInstall
+{
+public:
+	UE_API FGDKPlatformChunkInstall();
+	UE_API virtual ~FGDKPlatformChunkInstall();
+
+	//IPlatformChunkInstall implementation
+#if WITH_CHUNKINSTALL_ASYNC_INIT
+	UE_API virtual void AsyncInit( TFunction<void(bool/*bSuccess*/)> OnInitComplete ) override;
+#endif
+	UE_API virtual bool IsAvailable() const override;
+	UE_API virtual EChunkLocation::Type GetPakchunkLocation(int32 PakchunkIndex) override;
+	virtual bool GetProgressReportingTypeSupported(EChunkProgressReportingType::Type ReportType) override { return ReportType == EChunkProgressReportingType::PercentageComplete; }
+	UE_API virtual float GetChunkProgress( uint32 PakchunkIndex, EChunkProgressReportingType::Type ReportType ) override;
+	virtual EChunkInstallSpeed::Type GetInstallSpeed() override { return EChunkInstallSpeed::Slow; }
+	virtual bool SetInstallSpeed( EChunkInstallSpeed::Type InstallSpeed ) override { return false; }
+	UE_API virtual bool PrioritizePakchunk(int32 PakchunkIndex, EChunkPriority::Type Priority) override;
+	virtual bool DebugStartNextChunk() { return false; }
+	UE_API virtual void ExternalNotifyChunkAvailable(uint32 PakchunkIndex) override;
+
+	virtual bool SupportsNamedChunkInstall() const override { return true; }
+	UE_API virtual bool IsNamedChunkInProgress(const FName NamedChunk) override;
+	UE_API virtual bool InstallNamedChunks(const TArrayView<const FName>& NamedChunks) override;
+	UE_API virtual bool UninstallNamedChunks(const TArrayView<const FName>& NamedChunks) override;
+	UE_API virtual EChunkLocation::Type GetNamedChunkLocation(const FName NamedChunk) override;
+	UE_API virtual float GetNamedChunkProgress(const FName NamedChunk, EChunkProgressReportingType::Type ReportType) override;
+	UE_API virtual bool PrioritizeNamedChunk(const FName NamedChunk, EChunkPriority::Type Priority) override;
+	UE_API virtual ENamedChunkType GetNamedChunkType(const FName NamedChunk) const override;
+	UE_API virtual TArray<FName> GetNamedChunksByType(ENamedChunkType NamedChunkType) const override;
+
+	virtual bool SupportsBundleSource() const override { return true; }
+
+	UE_API virtual bool SetAutoPakMountingEnabled( bool bEnabled ) override;
+	UE_API virtual bool GetPakFilesInNamedChunk( const FName NamedChunk, TArray<FString>& OutFilesInChunk) const override;
+	UE_API virtual bool GetNamedChunkInstallationStatus( const FName NamedChunk, FChunkInstallationStatusDetail& OutChunkStatusDetail ) const override;
+	UE_API virtual bool IsNamedChunkForCurrentLocale( const FName NamedChunk ) const override;
+
+	// helper functions
+	static UE_API void DebugDumpChunkState();
+
+private:
+#if WITH_CHUNKINSTALL_ASYNC_INIT
+	std::atomic<bool> bIsInitialized = false;
+#endif
+
+	//IPlatformChunkInstall implementation
+	UE_API virtual bool PrioritizeChunk(uint32 UnrealChunkID, EChunkPriority::Type Priority) override;
+	UE_API virtual EChunkLocation::Type GetChunkLocation(uint32 UnrealChunkID) override;
+
+	// common chunk install functionality
+	bool bAutoMountPaks = true;
+	UE_API bool MountPaks(uint32 UnrealChunkID);
+	UE_API bool UnMountPaks(uint32 UnrealChunkID);
+
+	// package data
+	bool bIsPackaged;
+	bool bPackageHasFeatures;
+    char PackageIdentifier[XPACKAGE_IDENTIFIER_MAX_LENGTH];
+
+	// cached information about a chunk
+	struct FChunkMonitor
+	{
+		FChunkMonitor( FGDKPlatformChunkInstall& Owner, const XPackageChunkSelector& InSelector, XPackageChunkAvailability InInitialAvailability );
+		~FChunkMonitor();
+
+		FString GetDebugString() const;
+		FString GetName() const;
+		void OnProgressChange();
+
+		bool SetPriority(EChunkPriority::Type Priority);
+		float GetProgress(EChunkProgressReportingType::Type ReportType) const;
+		EChunkLocation::Type GetLocation() const;
+
+		bool bInstalled;
+		bool bWasInProgress;
+		uint8 CachedProgress; //0-100
+		TUniquePtr<char[]> SelectorTag;
+		XPackageChunkSelector Selector;
+		XPackageInstallationMonitorHandle InstallationMonitorHandle;
+		XTaskQueueRegistrationToken InstallationProgressChangedToken;
+		FGDKPlatformChunkInstall* Owner;
+	};
+
+	// chunk data
+	TArray<FName> OnDemandChunks;
+	TArray<FName> LanguageChunks;
+	TMap<uint32,TUniquePtr<FChunkMonitor>> ChunkMonitors;
+	TMap<FName,TUniquePtr<FChunkMonitor>> NamedChunkMonitors;
+	TMultiMap<FName,int32> InitialNamedChunksMap;
+	TSet<FName> SecondaryLanguageNamedChunks;
+	TSet<FName> UnavailableChunks;
+	FCriticalSection ChunkDataLock;
+
+	// chunk data functions
+#if USE_GDK_PACKAGE_MANIFEST_INIT
+	UE_API bool CreateChunkDataFromPackageManifest();
+#else
+	UE_API void CreateChunkDataForPackage(const XPackageDetails& Details);
+#endif
+	UE_API void OnNamedChunkInstallationChanged(const FName NamedChunk, bool bIsInstalled);
+	UE_API void OnChunkInstallationChanged(uint32 UnrealChunkID, bool bIsInstalled);
+	UE_API FChunkMonitor* FindChunkMonitor(uint32 UnrealChunkID);
+	UE_API FChunkMonitor* FindChunkMonitor(const FName NamedChunk);
+	const FChunkMonitor* FindChunkMonitor(uint32 UnrealChunkID) const   { return const_cast<FGDKPlatformChunkInstall*>(this)->FindChunkMonitor(UnrealChunkID); }
+	const FChunkMonitor* FindChunkMonitor(const FName NamedChunk) const { return const_cast<FGDKPlatformChunkInstall*>(this)->FindChunkMonitor(NamedChunk); }
+
+	// ticker & deferred operations
+	UE_API void AddToTicker();
+	UE_API void RemoveFromTicker();
+	UE_API bool Tick( float DeltaSeconds );
+	FTSTicker::FDelegateHandle TickHandle;
+	TSet<FName> PendingInstallChunks;
+	TSet<FName> PendingUninstallChunks;
+	TSet<FName> CurrentUninstallingChunks;
+
+	// helper class for creating an GDK XPackageChunkSelector array from an FName array
+	class FGDKCustomChunkSelector
+	{
+	public:
+		FGDKCustomChunkSelector( const TSet<FName>& NamedChunks, FGDKPlatformChunkInstall* Owner );
+		inline const FString& ToString() const					{ return StringTagList; }
+		inline const uint32 Num() const							{ return ChunkSelectors.Num(); }
+		inline XPackageChunkSelector* Get()						{ return ChunkSelectors.GetData(); }
+	
+	private:
+		FString							StringTagList;		//comma-separated list of tags, for logging only
+		TArray<XPackageChunkSelector>	ChunkSelectors;		//chunk selectors for the chunk tag ids
+	};
+};
+
+
+FString UE_API LexToString(const XPackageChunkSelector& ChunkSelector);
+FString UE_API LexToString(XPackageChunkAvailability Availability);
+
+#undef UE_API
+
+#else
+// dummy... needed?
+//#include "GenericPlatform/GenericPlatformChunkInstall.h"
+//typedef FGenericPlatformChunkInstall FGDKPlatformChunkInstall;
+
+#endif //WITH_GRDK
diff --git a/Engine/Plugins/Runtime/Microsoft/GDKVirtualKeyboard/GDKVirtualKeyboard.uplugin b/Engine/Plugins/Runtime/Microsoft/GDKVirtualKeyboard/GDKVirtualKeyboard.uplugin
new file mode 100644
--- /dev/null
+++ b/Engine/Plugins/Runtime/Microsoft/GDKVirtualKeyboard/GDKVirtualKeyboard.uplugin
@@ -0,0 +1,27 @@
+{
+	"FileVersion": 3,
+	"Version": 1,
+	"VersionName": "0.1",
+	"FriendlyName": "GDK Virtual Keyboard",
+	"Description": "Virtual Keyboard support for GDK. Used for Slate text input etc",
+	"CreatedBy": "Epic Games, Inc.",
+	"Category": "Platform",
+	"CreatedByURL": "https://epicgames.com",
+	"DocsURL": "",
+	"MarketplaceURL": "",
+	"SupportURL": "",
+	"EnabledByDefault": false,
+	"RequiresBuildPlatform": false,
+	"CanContainContent": false,
+	"Installed": false,
+	"SupportedTargetPlatforms": [ "Win64" ],
+	"Modules": [
+		{
+			"Name": "GDKVirtualKeyboard",
+			"Type": "RuntimeNoCommandlet",
+			"LoadingPhase": "Default",
+			"PlatformAllowList": [ "Win64" ],
+			"PlatformArchitectureDenyList": [ "Win64:arm64" ]
+		}
+	]
+}
\ No newline at end of file
diff --git a/Engine/Plugins/Runtime/Microsoft/GDKVirtualKeyboard/Source/GDKVirtualKeyboard/GDKVirtualKeyboard.Build.cs b/Engine/Plugins/Runtime/Microsoft/GDKVirtualKeyboard/Source/GDKVirtualKeyboard/GDKVirtualKeyboard.Build.cs
new file mode 100644
--- /dev/null
+++ b/Engine/Plugins/Runtime/Microsoft/GDKVirtualKeyboard/Source/GDKVirtualKeyboard/GDKVirtualKeyboard.Build.cs
@@ -0,0 +1,21 @@
+// Copyright Epic Games, Inc. All Rights Reserved.
+
+using UnrealBuildTool;
+using EpicGames.Core;
+
+public class GDKVirtualKeyboard : ModuleRules
+{
+	public GDKVirtualKeyboard(ReadOnlyTargetRules Target) : base(Target)
+	{
+		PrivateDependencyModuleNames.AddAll( 
+			"SlateCore",
+			"Slate",
+			"GDKRuntime",
+			"GRDK"
+		);
+
+		PublicDependencyModuleNames.AddAll(
+			"Core"
+		);
+	}
+}
diff --git a/Engine/Plugins/Runtime/Microsoft/GDKVirtualKeyboard/Source/GDKVirtualKeyboard/Private/GDKPlatformTextField.cpp b/Engine/Plugins/Runtime/Microsoft/GDKVirtualKeyboard/Source/GDKVirtualKeyboard/Private/GDKPlatformTextField.cpp
new file mode 100644
--- /dev/null
+++ b/Engine/Plugins/Runtime/Microsoft/GDKVirtualKeyboard/Source/GDKVirtualKeyboard/Private/GDKPlatformTextField.cpp
@@ -0,0 +1,243 @@
+// Copyright Epic Games, Inc. All Rights Reserved.
+
+#include "GDKPlatformTextField.h"
+
+#if WITH_GRDK
+#include "Widgets/Input/IVirtualKeyboardEntry.h"
+#include "Containers/Ticker.h"
+#include "HAL/PlatformProcess.h"
+#include "HAL/PlatformTime.h"
+#include "SlateGlobals.h"
+#include "SlateSettings.h"
+#include "UObject/Class.h"
+#include "Misc/ScopeLock.h"
+#include "GDKRuntimeModule.h"
+#include "Containers/Ticker.h"
+#include "Async/Async.h"
+
+#include "Microsoft/AllowMicrosoftPlatformTypes.h"
+#include <XGameUI.h>
+#include "Microsoft/HideMicrosoftPlatformTypes.h"
+
+#if !UE_BUILD_SHIPPING && !UE_BUILD_TEST
+static TAutoConsoleVariable<int32> CVarAutoCancelKeyboard(
+	TEXT("GDKPlatformTextField.AutoCancel"),
+	0,
+	TEXT("Automatically cancel platform text field shortly after opening it.\n")
+	TEXT(" 0: Do not cancel(default)\n")
+	TEXT("Otherwise: Approx number of seconds before cancelling\n"),
+	ECVF_Default
+);
+#endif
+
+static TAutoConsoleVariable<float> CVarGDKVirtualKeyboardDebounceDelaySeconds(
+	TEXT("GDKPlatformTextField.DebounceDelaySeconds"),
+	0.25f,
+	TEXT("Delay to prevent the platform text field from re-opening immediately after it has been closed")
+);
+
+struct FAsyncKeyboardEntryUserData
+{
+	FAsyncKeyboardEntryUserData( FGDKPlatformTextField* InTextField, TWeakPtr<IVirtualKeyboardEntry> InWeakTextEntryWidget )
+		: TextField(InTextField)
+		, WeakTextEntryWidget(InWeakTextEntryWidget)
+	{
+		IGDKRuntimeModule::Get().Internal_HandlePlatformTextFieldVisible(true);
+	}
+
+	~FAsyncKeyboardEntryUserData()
+	{
+		IGDKRuntimeModule::Get().Internal_HandlePlatformTextFieldVisible(false);
+	}
+
+	FGDKPlatformTextField* TextField;
+	TWeakPtr<IVirtualKeyboardEntry> WeakTextEntryWidget;
+};
+
+FGDKPlatformTextField::FGDKPlatformTextField()
+{
+	DebounceTime = 0.0f;
+}
+
+FGDKPlatformTextField::~FGDKPlatformTextField()
+{
+	KillExisitingDialog();
+}
+
+void FGDKPlatformTextField::KillExisitingDialog()
+{
+	check(IsInGameThread());
+
+	FScopeLock Lock(&AsyncBlockCrit);
+
+	// Lock is released here so we don't deadlock if the callback tries to acquire it while we're trying to cancel
+	if (KeyboardAsyncBlock.IsValid())
+	{
+		XAsyncCancel(KeyboardAsyncBlock->GetInnerBlockForGDKAPI());
+		(void)KeyboardAsyncBlock.Release();
+	}
+}
+
+void FGDKPlatformTextField::ShowVirtualKeyboard(bool bShow, int32 UserIndex, TSharedPtr<IVirtualKeyboardEntry> TextEntryWidget)
+{
+	bool bCanShow = (FPlatformTime::Seconds() >= DebounceTime);
+
+	KillExisitingDialog();
+
+	if (bShow && bCanShow)
+	{
+		XGameUiTextEntryInputScope InputScope = XGameUiTextEntryInputScope::Default;
+		switch (TextEntryWidget->GetVirtualKeyboardType())
+		{
+		case EKeyboardType::Keyboard_Number:
+			InputScope = XGameUiTextEntryInputScope::Number;
+			break;
+		case EKeyboardType::Keyboard_Web:
+			InputScope = XGameUiTextEntryInputScope::Url;
+			break;
+		case EKeyboardType::Keyboard_Email:
+			InputScope = XGameUiTextEntryInputScope::EmailSmtpAddress;
+			break;
+		case EKeyboardType::Keyboard_Password:
+			InputScope = XGameUiTextEntryInputScope::Password;
+			break;
+		case EKeyboardType::Keyboard_AlphaNumeric:
+			InputScope = XGameUiTextEntryInputScope::Alphanumeric;
+			break;
+		case EKeyboardType::Keyboard_Default:
+		default:
+			InputScope = XGameUiTextEntryInputScope::Default;
+			break;
+		}
+
+		check(KeyboardAsyncBlock.IsValid() == false);
+		
+		FAsyncKeyboardEntryUserData* UserData = new FAsyncKeyboardEntryUserData( this, TextEntryWidget );
+		KeyboardAsyncBlock = MakeUnique<FGDKAsyncBlock>(UserData, [](FGDKAsyncBlock* AsyncBlock)
+		{
+			// Take ownership of the userdata 
+			TUniquePtr<FAsyncKeyboardEntryUserData> UserData((FAsyncKeyboardEntryUserData*)AsyncBlock->GetUserData());
+			FGDKPlatformTextField* TextField = UserData->TextField;
+			TWeakPtr<IVirtualKeyboardEntry> WeakTextEntryWidget = UserData->WeakTextEntryWidget;
+
+			HRESULT Status = XAsyncGetStatus( AsyncBlock->GetInnerBlockForGDKAPI(), false );
+
+			// Only call into background thread functions, don't access any resources created on the main thread.
+			switch (Status)
+			{
+				case S_OK:
+				{
+					uint32 ResultSize = 0;
+					XGameUiShowTextEntryResultSize( AsyncBlock->GetInnerBlockForGDKAPI(), &ResultSize );						
+					char* ResultTextBuffer = new char[ResultSize+1];
+					XGameUiShowTextEntryResult( AsyncBlock->GetInnerBlockForGDKAPI(), ResultSize, ResultTextBuffer, nullptr );
+
+					FString StringResult(UTF8_TO_TCHAR(ResultTextBuffer));
+					delete[] ResultTextBuffer;
+
+					TextField->VirtualKeyboardCallbackBackgroundThread(WeakTextEntryWidget, StringResult, ETextEntryType::TextEntryAccepted);
+					break;
+				}
+				default:
+				{
+					UE_LOG(LogSlate, Verbose, TEXT("Virtual Keyboard returned without text due to an error state %x."), Status);
+					// Error falls through to the canceled case
+				}
+				case E_ABORT:
+				{
+					// Canceled returns original text and sends the TextEntryCanceld status
+					TextField->VirtualKeyboardCallbackBackgroundThread(WeakTextEntryWidget, TEXT(""), ETextEntryType::TextEntryCanceled);
+					break;
+				}
+			}
+
+			FScopeLock Lock(&TextField->AsyncBlockCrit);
+			// We may have raced with KillExisitingDialog so we acquire the lock to make sure they aren't trying to cancel us 
+			if (TextField->KeyboardAsyncBlock.Get() == AsyncBlock)
+			{
+				(void)TextField->KeyboardAsyncBlock.Release();
+			}
+
+			delete AsyncBlock;
+		});
+
+		HRESULT Result = XGameUiShowTextEntryAsync( KeyboardAsyncBlock->GetInnerBlockForGDKAPI(), 
+													TCHAR_TO_UTF8( *TextEntryWidget->GetHintText().ToString() ),
+													"",
+													TCHAR_TO_UTF8( *TextEntryWidget->GetText().ToString() ),
+													InputScope,
+													0 ); //0 = no limit
+		if( Result != S_OK )
+		{
+			// No synchronization here because we cancelled any outstanding task with KillExisitingDialog and 
+			// released KeyboardAsyncBlock so we own the pointer in there and any outstanding cancellation refers to a 
+			// different block
+			delete UserData;
+			UserData = nullptr;
+			KeyboardAsyncBlock.Reset();
+
+			if (Result == E_ACCESSDENIED)
+			{
+				// The application did not have focus. A system UI is up or we are constrained.
+				UE_LOG(LogSlate, Verbose, TEXT("Virtual Keyboard was invoked when the application did not have focus."));
+				return;
+			}
+
+			UE_LOG(LogSlate, Warning, TEXT("Virtual Keyboard failed with exception %x"), Result);
+			return;
+		}
+#if !UE_BUILD_SHIPPING && !UE_BUILD_TEST
+		else 
+		{
+			if (int32 Secs = CVarAutoCancelKeyboard.GetValueOnGameThread())
+			{
+				Async(EAsyncExecution::ThreadPool, [this, Secs]() {
+					FPlatformProcess::Sleep((float)Secs);
+					ExecuteOnGameThread(UE_SOURCE_LOCATION, [this]() { KillExisitingDialog();  });
+				});
+			}
+		}
+#endif
+	}
+	else if (bShow && !bCanShow)
+	{
+		UE_LOG(LogSlate, Warning, TEXT("Ignoring Virtual Keyboard show request because it was closed too recently"));
+	}
+
+}
+
+void FGDKPlatformTextField::VirtualKeyboardCallbackBackgroundThread(TWeakPtr<IVirtualKeyboardEntry> WeakTextEntryWidget, FString StringResult, ETextEntryType EntryType)
+{
+	ExecuteOnGameThread(UE_SOURCE_LOCATION, [this, WeakTextEntryWidget, StringResult, EntryType]() {
+		VirtualKeyboardCallbackGameThread(WeakTextEntryWidget, StringResult, EntryType);
+	});
+}
+
+void FGDKPlatformTextField::VirtualKeyboardCallbackGameThread(TWeakPtr<IVirtualKeyboardEntry> WeakTextEntryWidget, FString StringResult, ETextEntryType EntryType)
+{
+	check(IsInGameThread());
+	
+	// Don't touch KeyboardAsyncBlock here - it should have already been deleted by the callback or KillExisitingDialog
+	// If it's not null here, that means we re-opened the dialog and this callback doesn't own that block
+
+	if (TSharedPtr<IVirtualKeyboardEntry> StrongTextEntryWidget = WeakTextEntryWidget.Pin())
+	{
+		switch (EntryType)
+		{
+		case ETextEntryType::TextEntryAccepted:
+			StrongTextEntryWidget->SetTextFromVirtualKeyboard(FText::FromString(StringResult), ETextEntryType::TextEntryAccepted);
+			break;
+		case ETextEntryType::TextEntryCanceled:
+			StrongTextEntryWidget->SetTextFromVirtualKeyboard(StrongTextEntryWidget->GetText(), ETextEntryType::TextEntryCanceled);
+			break;
+		default:
+			check(false);
+		}
+	}
+
+	// prevent the virtual keyboard from opening again immediately
+	DebounceTime = FPlatformTime::Seconds() + CVarGDKVirtualKeyboardDebounceDelaySeconds.GetValueOnGameThread();
+}
+
+#endif //WITH_GRDK
+
diff --git a/Engine/Plugins/Runtime/Microsoft/GDKVirtualKeyboard/Source/GDKVirtualKeyboard/Private/GDKVirtualKeyboardModule.cpp b/Engine/Plugins/Runtime/Microsoft/GDKVirtualKeyboard/Source/GDKVirtualKeyboard/Private/GDKVirtualKeyboardModule.cpp
new file mode 100644
--- /dev/null
+++ b/Engine/Plugins/Runtime/Microsoft/GDKVirtualKeyboard/Source/GDKVirtualKeyboard/Private/GDKVirtualKeyboardModule.cpp
@@ -0,0 +1,40 @@
+// Copyright Epic Games, Inc. All Rights Reserved.
+
+#include "CoreMinimal.h"
+#include "Modules/ModuleInterface.h"
+#include "Modules/ModuleManager.h"
+#include "Features/IModularFeatures.h"
+#include "Framework/Application/IPlatformTextField.h"
+
+#if WITH_GRDK
+#include "GDKPlatformTextField.h"
+
+class FGDKVirtualKeyboardModule : public IModuleInterface, public IPlatformTextFieldFactory
+{
+public:
+
+	// IModuleInterface
+	virtual void StartupModule() override
+	{
+		IModularFeatures::Get().RegisterModularFeature(IPlatformTextFieldFactory::FeatureName, this);
+	}
+
+	virtual void ShutdownModule() override
+	{
+		IModularFeatures::Get().UnregisterModularFeature(IPlatformTextFieldFactory::FeatureName, this);
+	}
+
+	// IPlatformTextFieldFactory
+	virtual TUniquePtr<IPlatformTextField> CreateInstance() override
+	{
+		return MakeUnique<FGDKPlatformTextField>();
+	}
+};
+
+IMPLEMENT_MODULE(FGDKVirtualKeyboardModule, GDKVirtualKeyboard);
+
+#else
+
+IMPLEMENT_MODULE(FDefaultModuleImpl, GDKVirtualKeyboard);
+
+#endif //WITH_GRDK
diff --git a/Engine/Plugins/Runtime/Microsoft/GDKVirtualKeyboard/Source/GDKVirtualKeyboard/Public/GDKPlatformTextField.h b/Engine/Plugins/Runtime/Microsoft/GDKVirtualKeyboard/Source/GDKVirtualKeyboard/Public/GDKPlatformTextField.h
new file mode 100644
--- /dev/null
+++ b/Engine/Plugins/Runtime/Microsoft/GDKVirtualKeyboard/Source/GDKVirtualKeyboard/Public/GDKPlatformTextField.h
@@ -0,0 +1,29 @@
+// Copyright Epic Games, Inc. All Rights Reserved.
+
+#pragma once
+
+#include "Framework/Application/IPlatformTextField.h"
+
+#if WITH_GRDK
+#include "GDKTaskQueueHelpers.h"
+
+class FGDKPlatformTextField : public IPlatformTextField
+{
+public:
+	FGDKPlatformTextField();
+	~FGDKPlatformTextField();
+
+	virtual void ShowVirtualKeyboard(bool bShow, int32 UserIndex, TSharedPtr<IVirtualKeyboardEntry> TextEntryWidget) override;
+
+private:
+	void VirtualKeyboardCallbackBackgroundThread(TWeakPtr<IVirtualKeyboardEntry> WeakTextEntryWidget, FString StringResult, ETextEntryType EntryType);
+	void VirtualKeyboardCallbackGameThread(TWeakPtr<IVirtualKeyboardEntry> WeakTextEntryWidget, FString StringResult, ETextEntryType EntryType);
+
+	void KillExisitingDialog();
+
+	FCriticalSection AsyncBlockCrit;
+	TUniquePtr<FGDKAsyncBlock> KeyboardAsyncBlock;
+	double DebounceTime;
+};
+
+#endif //WITH_GRDK
