Fixed for on-demand chunk install progress was not being updated

#rnx

diff --git a/Engine/Plugins/Runtime/Microsoft/GDKPackageChunkInstall/Source/GDKPackageChunkInstall/Private/GDKPlatformChunkInstall.cpp b/Engine/Plugins/Runtime/Microsoft/GDKPackageChunkInstall/Source/GDKPackageChunkInstall/Private/GDKPlatformChunkInstall.cpp
--- a/Engine/Plugins/Runtime/Microsoft/GDKPackageChunkInstall/Source/GDKPackageChunkInstall/Private/GDKPlatformChunkInstall.cpp
+++ b/Engine/Plugins/Runtime/Microsoft/GDKPackageChunkInstall/Source/GDKPackageChunkInstall/Private/GDKPlatformChunkInstall.cpp
@@ -799,6 +799,73 @@ bool FGDKPlatformChunkInstall::IsNamedChunkInProgress(const FName NamedChunk)
 
 bool FGDKPlatformChunkInstall::Tick( float DeltaSeconds )
 {
+	// Active progress polling: Poll XPackageGetInstallationProgress for all chunk monitors
+	// that are not yet installed. This ensures progress is updated even if the
+	// XPackageRegisterInstallationProgressChanged callbacks are not firing properly.
+	{
+		FScopeLock Lock(&ChunkDataLock);
+		for (auto& Pair : NamedChunkMonitors)
+		{
+			FChunkMonitor* ChunkMonitor = Pair.Value.Get();
+			if (ChunkMonitor && !ChunkMonitor->bInstalled && ChunkMonitor->InstallationMonitorHandle != nullptr)
+			{
+				XPackageInstallationProgress Progress;
+				XPackageGetInstallationProgress(ChunkMonitor->InstallationMonitorHandle, &Progress);
+
+				uint8 NewProgress = (Progress.totalBytes > 0) ? static_cast<uint8>((Progress.installedBytes * 100) / Progress.totalBytes) : 0;
+
+				// Only update and log if progress changed
+				if (NewProgress != ChunkMonitor->CachedProgress)
+				{
+					ChunkMonitor->CachedProgress = NewProgress;
+					UE_LOG(LogChunkInstaller, Verbose, TEXT("%s installation progress polled: %llu / %llu bytes (%d%%).%s"),
+						*ChunkMonitor->GetName(), Progress.installedBytes, Progress.totalBytes, NewProgress, Progress.completed ? TEXT(" completed") : TEXT(""));
+				}
+
+				// Check for completion - must notify via OnNamedChunkInstallationChanged just like OnProgressChange does
+				if (Progress.completed && !ChunkMonitor->bInstalled && Progress.totalBytes > 0)
+				{
+					ChunkMonitor->bInstalled = true;
+					ChunkMonitor->bWasInProgress = false;
+					UE_LOG(LogChunkInstaller, Display, TEXT("%s installation completed via polling"), *ChunkMonitor->GetName());
+					// Notify dependent systems - Pair.Key is the FName for named chunks
+					OnNamedChunkInstallationChanged(Pair.Key, true);
+				}
+			}
+		}
+
+		// Also check indexed chunk monitors
+		for (auto& Pair : ChunkMonitors)
+		{
+			FChunkMonitor* ChunkMonitor = Pair.Value.Get();
+			if (ChunkMonitor && !ChunkMonitor->bInstalled && ChunkMonitor->InstallationMonitorHandle != nullptr)
+			{
+				XPackageInstallationProgress Progress;
+				XPackageGetInstallationProgress(ChunkMonitor->InstallationMonitorHandle, &Progress);
+
+				uint8 NewProgress = (Progress.totalBytes > 0) ? static_cast<uint8>((Progress.installedBytes * 100) / Progress.totalBytes) : 0;
+
+				// Only update and log if progress changed
+				if (NewProgress != ChunkMonitor->CachedProgress)
+				{
+					ChunkMonitor->CachedProgress = NewProgress;
+					UE_LOG(LogChunkInstaller, Verbose, TEXT("%s installation progress polled: %llu / %llu bytes (%d%%).%s"),
+						*ChunkMonitor->GetName(), Progress.installedBytes, Progress.totalBytes, NewProgress, Progress.completed ? TEXT(" completed") : TEXT(""));
+				}
+
+				// Check for completion - must notify via OnChunkInstallationChanged just like OnProgressChange does
+				if (Progress.completed && !ChunkMonitor->bInstalled && Progress.totalBytes > 0)
+				{
+					ChunkMonitor->bInstalled = true;
+					ChunkMonitor->bWasInProgress = false;
+					UE_LOG(LogChunkInstaller, Display, TEXT("%s installation completed via polling"), *ChunkMonitor->GetName());
+					// Notify dependent systems - Pair.Key is the UnrealChunkID for indexed chunks
+					OnChunkInstallationChanged(Pair.Key, true);
+				}
+			}
+		}
+	}
+
 	if (PendingInstallChunks.Num() > 0)
 	{
 		// async callback
