Take UINT64_MAX into account when testing XPackageInstallationProgress.totalBytes

#rnx
#Synced-CL 52123119

diff --git a/Engine/Plugins/Runtime/Microsoft/GDKPackageChunkInstall/Source/GDKPackageChunkInstall/Private/GDKPlatformChunkInstall.cpp b/Engine/Plugins/Runtime/Microsoft/GDKPackageChunkInstall/Source/GDKPackageChunkInstall/Private/GDKPlatformChunkInstall.cpp
--- a/Engine/Plugins/Runtime/Microsoft/GDKPackageChunkInstall/Source/GDKPackageChunkInstall/Private/GDKPlatformChunkInstall.cpp
+++ b/Engine/Plugins/Runtime/Microsoft/GDKPackageChunkInstall/Source/GDKPackageChunkInstall/Private/GDKPlatformChunkInstall.cpp
@@ -1412,7 +1412,7 @@ void FGDKPlatformChunkInstall::FChunkMonitor::OnProgressChange()
 		bInstalled = true;
 		bComplete = true;
 	}
-	else if ((bInstalled || bWasInProgress) && (Progress.installedBytes == 0))
+	else if ((bInstalled || bWasInProgress) && (Progress.installedBytes == 0) && (Progress.totalBytes != UINT64_MAX))
 	{
 		bInstalled = false;
 		bComplete = true;
